Fixes on the Constraint Node

diff --git a/Engine/Plugins/Experimental/ControlRig/Source/ControlRig/Private/Units/RigUnit_TransformConstraint.cpp b/Engine/Plugins/Experimental/ControlRig/Source/ControlRig/Private/Units/RigUnit_TransformConstraint.cpp
--- a/Engine/Plugins/Experimental/ControlRig/Source/ControlRig/Private/Units/RigUnit_TransformConstraint.cpp
+++ b/Engine/Plugins/Experimental/ControlRig/Source/ControlRig/Private/Units/RigUnit_TransformConstraint.cpp
@@ -66,7 +66,7 @@ void FRigUnit_TransformConstraint::Execute(const FRigUnitContext& InContext)
 			if (JointIndex != INDEX_NONE)
 			{
 				const int32 TargetNum = Targets.Num();
-				if (TargetNum > 0)
+				if (TargetNum > 0 && ConstraintData.Num() > 0)
 				{
 					for (int32 ConstraintIndex= 0; ConstraintIndex< ConstraintData.Num(); ++ConstraintIndex)
 					{
diff --git a/Engine/Source/Runtime/AnimationCore/Private/AnimationCoreLibrary.cpp b/Engine/Source/Runtime/AnimationCore/Private/AnimationCoreLibrary.cpp
--- a/Engine/Source/Runtime/AnimationCore/Private/AnimationCoreLibrary.cpp
+++ b/Engine/Source/Runtime/AnimationCore/Private/AnimationCoreLibrary.cpp
@@ -142,7 +142,7 @@ FQuat SolveAim(const FTransform& CurrentTransform, const FVector& TargetPosition
 FTransform SolveConstraints(const FTransform& CurrentTransform, const FTransform& CurrentParentTransform, const TArray<struct FConstraintData>& Constraints)
 {
 	int32 TotalNum = Constraints.Num();
-	ensureAlways(TotalNum > 0);
+	ensure(TotalNum > 0);
 
 	FMultiTransformBlendHelper BlendHelperInLocalSpace;
 	FTransform BlendedLocalTransform = CurrentTransform.GetRelativeTransform(CurrentParentTransform);
