Support for animation data source registry. This mechanism will be used within the ControlRig to provide access to external data sources.

diff --git a/Engine/Source/Runtime/AnimationCore/Private/AnimationDataSource.cpp b/Engine/Source/Runtime/AnimationCore/Private/AnimationDataSource.cpp
new file mode 100644
--- /dev/null
+++ b/Engine/Source/Runtime/AnimationCore/Private/AnimationDataSource.cpp
@@ -0,0 +1,41 @@
+// Copyright 1998-2019 Epic Games, Inc. All Rights Reserved.
+
+#include "AnimationDataSource.h"
+
+bool UAnimationDataSourceRegistry::RegisterDataSource(const FName& InName, UObject* InDataSource)
+{
+	if (DataSources.Contains(InName))
+	{
+		return false;
+	}
+	DataSources.Add(InName, InDataSource);
+	return true;
+}
+
+bool UAnimationDataSourceRegistry::UnregisterDataSource(const FName& InName)
+{
+	return DataSources.Remove(InName) > 0;
+}
+
+bool UAnimationDataSourceRegistry::ContainsSource(const FName& InName) const
+{
+	return DataSources.Contains(InName);
+}
+
+UObject* UAnimationDataSourceRegistry::RequestSource(const FName& InName, UClass* InExpectedClass) const
+{
+	UObject* const* DataSource = DataSources.Find(InName);
+	if (DataSource == nullptr)
+	{
+		return nullptr;
+	}
+	if (*DataSource == nullptr)
+	{
+		return nullptr;
+	}
+	if (!(*DataSource)->IsA(InExpectedClass))
+	{
+		return nullptr;
+	}
+	return *DataSource;
+}
\ No newline at end of file
diff --git a/Engine/Source/Runtime/AnimationCore/Public/AnimationDataSource.h b/Engine/Source/Runtime/AnimationCore/Public/AnimationDataSource.h
new file mode 100644
--- /dev/null
+++ b/Engine/Source/Runtime/AnimationCore/Public/AnimationDataSource.h
@@ -0,0 +1,74 @@
+// Copyright 1998-2019 Epic Games, Inc. All Rights Reserved.
+
+#pragma once
+
+#include "CoreMinimal.h"
+#include "UObject/ObjectMacros.h"
+#include "UObject/Interface.h"
+#include "AnimationDataSource.generated.h"
+
+UCLASS()
+class ANIMATIONCORE_API UAnimationDataSourceRegistry : public UObject
+{
+	GENERATED_BODY()
+
+public:
+
+	/**
+	 * Registers a new data source under a given name.
+	 * Returns false in case the data source has already been registered.
+	 *
+	 * @param InName         The name to register the data source under
+	 * @param InDataSource   The data source to register
+	 * @return true if succeeded
+	 */
+	bool RegisterDataSource(const FName& InName, UObject* InDataSource);
+
+	/**
+	 * Unregisters / removes an existing data source under a given name.
+	 * Returns false in case the data source has already been unregistered.
+	 *
+	 * @param InName The name of the data source to remove
+	 * @return true if succeeded
+	 */
+	bool UnregisterDataSource(const FName& InName);
+
+	/**
+	 * Returns true if this registry contains a source with the given name
+	 *
+	 * @param InName The name of the data source to look up.
+	 * @return true if a source with the given name exists
+	 */
+	bool ContainsSource(const FName& InName) const;
+
+	/**
+	 * Returns a given data source and cast it to the expected class.
+	 *
+	 * @param InName The name of the data source to look up.
+	 * @param InExpectedClass The class expected from the data source
+	 * @return The requested data source
+	 */
+	UObject* RequestSource(const FName& InName, UClass* InExpectedClass) const;
+
+	/**
+	 * Returns a given data source and cast it to the expected class.
+	 * 
+	 * @param InName The name of the data source to look up.
+	 * @return The requested data source
+	 */
+	template<class T>
+	T* RequestSource(const FName& InName) const
+	{
+		UObject* DataSource = RequestSource(InName, T::StaticClass());
+		if (DataSource == nullptr)
+		{
+			return nullptr;
+		}
+		return Cast<T>(DataSource);
+	}
+
+private:
+
+	UPROPERTY()
+	TMap<FName, UObject*> DataSources;
+};
\ No newline at end of file
