Control Rig: MultiplexStorage for fragmented runtime memory + unit tests for both the base data structure and the initial control rig usage.

Adaption of unit tests

diff --git a/Engine/Source/Runtime/AnimationCore/Private/MultiplexStorage.cpp b/Engine/Source/Runtime/AnimationCore/Private/MultiplexStorage.cpp
new file mode 100644
--- /dev/null
+++ b/Engine/Source/Runtime/AnimationCore/Private/MultiplexStorage.cpp
@@ -0,0 +1,186 @@
+// Copyright 1998-2019 Epic Games, Inc. All Rights Reserved.
+
+#include "MultiplexStorage.h"
+
+FMultiplexStorage::FMultiplexStorage()
+{
+}
+
+FMultiplexStorage::~FMultiplexStorage()
+{
+}
+
+FMultiplexStorage& FMultiplexStorage::operator= (const FMultiplexStorage &InOther)
+{
+	Data.Reset();
+	Data.Append(InOther.Data);
+	Addresses.Reset();
+	Addresses.Append(InOther.Addresses);
+	UpdateAddresses();
+	return *this;
+}
+
+#if WITH_EDITORONLY_DATA
+
+void FMultiplexStorage::Reset()
+{
+
+}
+
+int32 FMultiplexStorage::Add(const FName& InNewName, int32 InElementSize, int32 InCount, const void* InData)
+{
+	ensure(InElementSize > 0 && InCount > 0);
+
+	if (!IsNameAvailable(InNewName))
+	{
+		return INDEX_NONE;
+	}
+
+	FMultiplexAddress NewAddress;
+	NewAddress.ByteIndex = Data.Num();
+	NewAddress.Name = InNewName;
+	NewAddress.ElementSize = InElementSize;
+	NewAddress.ElementCount = InCount;
+
+	Data.AddZeroed(NewAddress.NumBytes());
+
+	if (InData != nullptr)
+	{
+		FMemory::Memcpy(&Data[NewAddress.ByteIndex], InData, NewAddress.NumBytes());
+	}
+
+	int32 AddressIndex = Addresses.Num();
+	Addresses.Add(NewAddress);
+	UpdateAddresses();
+	return AddressIndex;
+}
+
+bool FMultiplexStorage::Remove(int32 InAddressIndex)
+{
+	if (InAddressIndex < 0 || InAddressIndex >= Addresses.Num())
+	{
+		return false;
+	}
+
+	FMultiplexAddress AddressToRemove = Addresses[InAddressIndex];
+	Data.RemoveAt(AddressToRemove.ByteIndex, AddressToRemove.NumBytes());
+	Addresses.RemoveAt(InAddressIndex);
+
+	for (int32 Index = InAddressIndex; Index < Addresses.Num(); Index++)
+	{
+		Addresses[Index].ByteIndex -= AddressToRemove.NumBytes();
+	}
+
+	UpdateAddresses();
+	return true;
+}
+
+bool FMultiplexStorage::Remove(const FName& InAddressName)
+{
+	return Remove(GetIndex(InAddressName));
+}
+
+FName FMultiplexStorage::Rename(int32 InAddressIndex, const FName& InNewName)
+{
+	if (Addresses[InAddressIndex].Name == InNewName)
+	{
+		return Addresses[InAddressIndex].Name;
+	}
+
+	if (!IsNameAvailable(InNewName))
+	{
+		return Addresses[InAddressIndex].Name;
+	}
+
+	Addresses[InAddressIndex].Name = InNewName;
+	UpdateAddresses();
+
+	return InNewName;
+}
+
+FName FMultiplexStorage::Rename(const FName& InOldName, const FName& InNewName)
+{
+	int32 AddressIndex = GetIndex(InOldName);
+	if (AddressIndex == INDEX_NONE)
+	{
+		return NAME_None;
+	}
+
+	return Rename(AddressIndex, InNewName);
+}
+
+bool FMultiplexStorage::Resize(int32 InAddressIndex, int32 InNewElementCount)
+{
+	if (InNewElementCount <= 0)
+	{
+		return Remove(InAddressIndex);
+	}
+
+	if (Addresses[InAddressIndex].ElementCount == InNewElementCount)
+	{
+		return false;
+	}
+
+	if (Addresses[InAddressIndex].ElementCount > InNewElementCount) // shrink
+	{
+		int32 ElementsToRemove = Addresses[InAddressIndex].ElementCount - InNewElementCount;
+		int32 NumBytesToRemove = Addresses[InAddressIndex].ElementSize * ElementsToRemove;
+		int32 FirstByteToRemove = Addresses[InAddressIndex].ByteIndex + Addresses[InAddressIndex].ElementSize * InNewElementCount;
+
+		Data.RemoveAt(FirstByteToRemove, NumBytesToRemove);
+		Addresses[InAddressIndex].ElementCount = InNewElementCount;
+
+		for (int32 AddressIndex = InAddressIndex + 1; AddressIndex < Addresses.Num(); AddressIndex++)
+		{
+			Addresses[AddressIndex].ByteIndex -= NumBytesToRemove;
+		}
+	}
+	else // grow
+	{
+		int32 ElementsToAdd = InNewElementCount - Addresses[InAddressIndex].ElementCount;
+		int32 NumBytesToAdd = Addresses[InAddressIndex].ElementSize * ElementsToAdd;
+		int32 FirstByteToAdd = Addresses[InAddressIndex].ByteIndex + Addresses[InAddressIndex].ElementSize * Addresses[InAddressIndex].ElementCount;
+
+		Data.InsertZeroed(FirstByteToAdd, NumBytesToAdd);
+		Addresses[InAddressIndex].ElementCount = InNewElementCount;
+
+		for (int32 AddressIndex = InAddressIndex + 1; AddressIndex < Addresses.Num(); AddressIndex++)
+		{
+			Addresses[AddressIndex].ByteIndex += NumBytesToAdd;
+		}
+	}
+
+	UpdateAddresses();
+	return true;
+}
+
+bool FMultiplexStorage::Resize(const FName& InAddressName, int32 InNewElementCount)
+{
+	int32 AddressIndex = GetIndex(InAddressName);
+	if (AddressIndex == INDEX_NONE)
+	{
+		return false;
+	}
+
+	return Resize(AddressIndex, InNewElementCount);
+}
+
+#endif
+
+void FMultiplexStorage::UpdateAddresses()
+{
+	for (FMultiplexAddress& Address : Addresses)
+	{
+		Address.Pointer = &Data[Address.ByteIndex];
+	}
+
+#if WITH_EDITORONLY_DATA
+
+	NameMap.Reset();
+	for (int32 Index = 0; Index < Addresses.Num(); Index++)
+	{
+		NameMap.Add(Addresses[Index].Name, Index);
+	}
+
+#endif
+}
diff --git a/Engine/Source/Runtime/AnimationCore/Public/MultiplexStorage.h b/Engine/Source/Runtime/AnimationCore/Public/MultiplexStorage.h
new file mode 100644
--- /dev/null
+++ b/Engine/Source/Runtime/AnimationCore/Public/MultiplexStorage.h
@@ -0,0 +1,219 @@
+// Copyright 1998-2019 Epic Games, Inc. All Rights Reserved.
+
+/*=============================================================================
+	AngularLimit.h: Angular limit features
+=============================================================================*/
+
+#pragma once
+
+#include "CoreMinimal.h"
+#include "UObject/ObjectMacros.h"
+#include "MultiplexStorage.generated.h"
+
+USTRUCT()
+struct ANIMATIONCORE_API FMultiplexAddress
+{
+	GENERATED_USTRUCT_BODY()
+
+	FMultiplexAddress()
+		: Pointer(nullptr)
+		, ByteIndex(INDEX_NONE)
+		, ElementSize(0)
+		, ElementCount(0)
+#if WITH_EDITORONLY_DATA
+		, Name(NAME_None)
+#endif
+	{
+	}
+
+	void* Pointer;
+
+	UPROPERTY()
+	int32 ByteIndex;
+
+	UPROPERTY()
+	int32 ElementSize;
+
+	UPROPERTY()
+	int32 ElementCount;
+
+#if WITH_EDITORONLY_DATA
+	UPROPERTY()
+		FName Name;
+#endif
+
+	FORCEINLINE bool IsArray() const { return ElementCount > 1; }
+	FORCEINLINE int32 NumBytes() const { return ElementCount * ElementSize;  }
+
+	template<class T>
+	FORCEINLINE const T* Get() const
+	{
+		ensure(sizeof(T) == ElementSize);
+		ensure(ElementCount > 0);
+		return (const T*)Pointer;
+	}
+
+	template<class T>
+	FORCEINLINE const T& GetRef() const
+	{
+		return *Get<T>();
+	}
+
+	template<class T>
+	FORCEINLINE T* Get()
+	{
+		ensure(sizeof(T) == ElementSize);
+		ensure(ElementCount > 0);
+		return (T*)Pointer;
+	}
+
+	template<class T>
+	FORCEINLINE T& GetRef()
+	{
+		return *Get<T>();
+	}
+
+	template<class T>
+	FORCEINLINE TArrayView<T> GetArray()
+	{
+		ensure(sizeof(T) == ElementSize);
+		ensure(ElementCount > 0);
+		return TArrayView<T>((T*)Pointer, ElementCount);
+	}
+};
+
+USTRUCT()
+struct ANIMATIONCORE_API FMultiplexStorage
+{
+	GENERATED_USTRUCT_BODY()
+	
+public:
+
+	FMultiplexStorage();
+	~FMultiplexStorage();
+
+	FMultiplexStorage& operator= (const FMultiplexStorage &InOther);
+
+	FORCEINLINE int32 Num() const { return Addresses.Num(); }
+	FORCEINLINE const FMultiplexAddress& operator[](int32 InIndex) const { return Addresses[InIndex]; }
+	FORCEINLINE FMultiplexAddress& operator[](int32 InIndex) { return Addresses[InIndex]; }
+#if WITH_EDITORONLY_DATA
+	FORCEINLINE const FMultiplexAddress& operator[](const FName& InName) const { return Addresses[GetIndex(InName)]; }
+	FORCEINLINE FMultiplexAddress& operator[](const FName& InName) { return Addresses[GetIndex(InName)]; }
+#endif
+
+	FORCEINLINE TArray<FMultiplexAddress>::RangedForIteratorType      begin() { return Addresses.begin(); }
+	FORCEINLINE TArray<FMultiplexAddress>::RangedForConstIteratorType begin() const { return Addresses.begin(); }
+	FORCEINLINE TArray<FMultiplexAddress>::RangedForIteratorType      end() { return Addresses.end(); }
+	FORCEINLINE TArray<FMultiplexAddress>::RangedForConstIteratorType end() const { return Addresses.end(); }
+
+	template<class T>
+	FORCEINLINE const T* Get(int32 InAddressIndex) const
+	{
+		const FMultiplexAddress& Address = Addresses[InAddressIndex];
+		ensure(sizeof(T) == Address.ElementSize);
+		ensure(Address.ElementCount > 0);
+		return (const T*)&Data[Address.ByteIndex];
+	}
+
+	template<class T>
+	FORCEINLINE const T& GetRef(int32 InAddressIndex) const
+	{
+		return *Get<T>(InAddressIndex);
+	}
+
+	template<class T>
+	FORCEINLINE T* Get(int32 InAddressIndex)
+	{
+		const FMultiplexAddress& Address = Addresses[InAddressIndex];
+		ensure(sizeof(T) == Address.ElementSize);
+		ensure(Address.ElementCount > 0);
+		return (T*)&Data[Address.ByteIndex];
+	}
+
+	template<class T>
+	FORCEINLINE T& GetRef(int32 InAddressIndex)
+	{
+		return *Get<T>(InAddressIndex);
+	}
+
+	template<class T>
+	FORCEINLINE TArrayView<T> GetArray(int32 InAddressIndex)
+	{
+		const FMultiplexAddress& Address = Addresses[InAddressIndex];
+		ensure(sizeof(T) == Address.ElementSize);
+		ensure(Address.ElementCount > 0);
+		return TArrayView<T>((T*)&Data[Address.ByteIndex], Address.ElementCount);
+	}
+
+#if WITH_EDITORONLY_DATA
+
+	FORCEINLINE int32 GetIndex(const FName& InName) const
+	{
+		if (NameMap.Num() != Addresses.Num())
+		{
+			for (int32 Index = 0; Index < Addresses.Num(); Index++)
+			{
+				if (Addresses[Index].Name == InName)
+				{
+					return Index;
+				}
+			}
+		}
+		else
+		{
+			const int32* Index = NameMap.Find(InName);
+			if (Index != nullptr)
+			{
+				return *Index;
+			}
+		}
+
+		return INDEX_NONE;
+	}
+
+	FORCEINLINE bool IsNameAvailable(const FName& InPotentialNewName) const { return GetIndex(InPotentialNewName) == INDEX_NONE; }
+
+	void Reset();
+
+	int32 Add(const FName& InNewName, int32 InElementSize, int32 InCount, const void* InData = nullptr);
+
+	template<class T>
+	int32 Add(const FName& InNewName, const T& InValue)
+	{
+		return Add(InNewName, sizeof(T), 1, (const void*)&InValue);
+	}
+
+	template<class T>
+	int32 AddArray(const FName& InNewName, const TArray<T>& InArray)
+	{
+		return Add(InNewName, sizeof(T), InArray.Num(), (const void*)InArray.GetData());
+	}
+
+	bool Remove(int32 InAddressIndex);
+	bool Remove(const FName& InAddressName);
+	FName Rename(int32 InAddressIndex, const FName& InNewName);
+	FName Rename(const FName& InOldName, const FName& InNewName);
+	bool Resize(int32 InAddressIndex, int32 InNewElementCount);
+	bool Resize(const FName& InAddressName, int32 InNewElementCount);
+
+#endif
+
+	void UpdateAddresses();
+
+private:
+
+	// disable copy constructor
+	FMultiplexStorage(const FMultiplexStorage& Other) {}
+
+	UPROPERTY()
+	TArray<FMultiplexAddress> Addresses;
+
+	UPROPERTY()
+	TArray<uint8> Data;
+
+#if WITH_EDITORONLY_DATA
+	UPROPERTY(transient)
+	TMap<FName, int32> NameMap;
+#endif
+};
