MultiplexStorage: Implement support for names and conscious mode for no-name-support + unit tests

diff --git a/Engine/Source/Runtime/AnimationCore/Private/MultiplexStorage.cpp b/Engine/Source/Runtime/AnimationCore/Private/MultiplexStorage.cpp
--- a/Engine/Source/Runtime/AnimationCore/Private/MultiplexStorage.cpp
+++ b/Engine/Source/Runtime/AnimationCore/Private/MultiplexStorage.cpp
@@ -2,7 +2,8 @@
 
 #include "MultiplexStorage.h"
 
-FMultiplexStorage::FMultiplexStorage()
+FMultiplexStorage::FMultiplexStorage(bool bInUseNames)
+	:bUseNameMap(bInUseNames)
 {
 }
 
@@ -12,6 +13,7 @@ FMultiplexStorage::~FMultiplexStorage()
 
 FMultiplexStorage& FMultiplexStorage::operator= (const FMultiplexStorage &InOther)
 {
+	bUseNameMap = InOther.bUseNameMap;
 	Data.Reset();
 	Data.Append(InOther.Data);
 	Addresses.Reset();
@@ -20,25 +22,125 @@ FMultiplexStorage& FMultiplexStorage::operator= (const FMultiplexStorage &InOthe
 	return *this;
 }
 
-#if WITH_EDITORONLY_DATA
-
 void FMultiplexStorage::Reset()
 {
+	Data.Reset();
+	Addresses.Reset();
+	NameMap.Reset();
+}
+
+bool FMultiplexStorage::Copy(
+	int32 InSourceAddressIndex,
+	int32 InTargetAddressIndex,
+	int32 InSourceByteOffset,
+	int32 InTargetByteOffset,
+	int32 InNumBytes)
+{
+	ensure(Addresses.IsValidIndex(InSourceAddressIndex));
+	ensure(Addresses.IsValidIndex(InTargetAddressIndex));
+
+	if (InSourceAddressIndex == InTargetAddressIndex && InSourceByteOffset == InTargetByteOffset)
+	{
+		return false;
+	}
+
+	const FMultiplexAddress& Source = Addresses[InSourceAddressIndex];
+	const FMultiplexAddress& Target = Addresses[InTargetAddressIndex];
+
+	int32 SourceStartByte = Source.ByteIndex;
+	int32 SourceNumBytes = Source.NumBytes();
+	if(InSourceByteOffset != INDEX_NONE)
+	{
+		ensure(InNumBytes != INDEX_NONE);
+		ensure(InSourceByteOffset >= 0 && InSourceByteOffset < Source.NumBytes());
+		ensure(InNumBytes > 0 && InSourceByteOffset + InNumBytes <= Source.NumBytes());
+
+		SourceStartByte += InSourceByteOffset;
+		SourceNumBytes = InNumBytes;
+	}
+
+	int32 TargetStartByte = Target.ByteIndex;
+	int32 TargetNumBytes = Target.NumBytes();
+	if(InTargetByteOffset != INDEX_NONE)
+	{
+		ensure(InNumBytes != INDEX_NONE);
+		ensure(InTargetByteOffset >= 0 && InTargetByteOffset < Target.NumBytes());
+		ensure(InNumBytes > 0 && InTargetByteOffset + InNumBytes <= Target.NumBytes());
+
+		TargetStartByte += InTargetByteOffset;
+		TargetNumBytes = InNumBytes;
+	}
 
+	if (SourceNumBytes != TargetNumBytes)
+	{
+		return false;
+	}
+
+	FMemory::Memcpy(&Data[TargetStartByte], &Data[SourceStartByte], TargetNumBytes);
+
+	return true;
+}
+
+bool FMultiplexStorage::Copy(
+	const FName& InSourceName,
+	const FName& InTargetName,
+	int32 InSourceByteOffset,
+	int32 InTargetByteOffset,
+	int32 InNumBytes)
+{
+	ensure(bUseNameMap);
+
+	int32 SourceAddressIndex = GetIndex(InSourceName);
+	int32 TargetAddressIndex = GetIndex(InTargetName);
+
+	if(SourceAddressIndex == INDEX_NONE || TargetAddressIndex == INDEX_NONE)
+	{
+		return false;
+	}
+
+	return Copy(SourceAddressIndex, TargetAddressIndex, InSourceByteOffset, InTargetByteOffset, InNumBytes);
+}
+
+int32 FMultiplexStorage::Add(int32 InElementSize, int32 InCount, const void* InData)
+{
+	FName Name = NAME_None;
+
+	if (bUseNameMap)
+	{
+		const TCHAR* AddressPrefix = TEXT("Address");
+		int32 AddressSuffix = 0;
+		do
+		{
+			Name = FName(*FString::Printf(TEXT("%s_%d"), AddressPrefix, AddressSuffix++));
+		}
+		while (!IsNameAvailable(Name));
+	}
+
+	return Add(Name, InElementSize, InCount, InData);
 }
 
 int32 FMultiplexStorage::Add(const FName& InNewName, int32 InElementSize, int32 InCount, const void* InData)
 {
+	if (bUseNameMap)
+	{
+		ensure(InNewName != NAME_None);
+	}
 	ensure(InElementSize > 0 && InCount > 0);
 
-	if (!IsNameAvailable(InNewName))
+	if (bUseNameMap)
 	{
-		return INDEX_NONE;
+		if (!IsNameAvailable(InNewName))
+		{
+			return INDEX_NONE;
+		}
 	}
 
 	FMultiplexAddress NewAddress;
 	NewAddress.ByteIndex = Data.Num();
-	NewAddress.Name = InNewName;
+	if (bUseNameMap)
+	{
+		NewAddress.Name = InNewName;
+	}
 	NewAddress.ElementSize = InElementSize;
 	NewAddress.ElementCount = InCount;
 
@@ -77,6 +179,7 @@ bool FMultiplexStorage::Remove(int32 InAddressIndex)
 
 bool FMultiplexStorage::Remove(const FName& InAddressName)
 {
+	ensure(bUseNameMap);
 	return Remove(GetIndex(InAddressName));
 }
 
@@ -100,6 +203,8 @@ FName FMultiplexStorage::Rename(int32 InAddressIndex, const FName& InNewName)
 
 FName FMultiplexStorage::Rename(const FName& InOldName, const FName& InNewName)
 {
+	ensure(bUseNameMap);
+
 	int32 AddressIndex = GetIndex(InOldName);
 	if (AddressIndex == INDEX_NONE)
 	{
@@ -156,6 +261,8 @@ bool FMultiplexStorage::Resize(int32 InAddressIndex, int32 InNewElementCount)
 
 bool FMultiplexStorage::Resize(const FName& InAddressName, int32 InNewElementCount)
 {
+	ensure(bUseNameMap);
+
 	int32 AddressIndex = GetIndex(InAddressName);
 	if (AddressIndex == INDEX_NONE)
 	{
@@ -165,8 +272,6 @@ bool FMultiplexStorage::Resize(const FName& InAddressName, int32 InNewElementCou
 	return Resize(AddressIndex, InNewElementCount);
 }
 
-#endif
-
 void FMultiplexStorage::UpdateAddresses()
 {
 	for (FMultiplexAddress& Address : Addresses)
@@ -174,13 +279,12 @@ void FMultiplexStorage::UpdateAddresses()
 		Address.Pointer = &Data[Address.ByteIndex];
 	}
 
-#if WITH_EDITORONLY_DATA
-
-	NameMap.Reset();
-	for (int32 Index = 0; Index < Addresses.Num(); Index++)
+	if (bUseNameMap)
 	{
-		NameMap.Add(Addresses[Index].Name, Index);
+		NameMap.Reset();
+		for (int32 Index = 0; Index < Addresses.Num(); Index++)
+		{
+			NameMap.Add(Addresses[Index].Name, Index);
+		}
 	}
-
-#endif
 }
diff --git a/Engine/Source/Runtime/AnimationCore/Public/MultiplexStorage.h b/Engine/Source/Runtime/AnimationCore/Public/MultiplexStorage.h
--- a/Engine/Source/Runtime/AnimationCore/Public/MultiplexStorage.h
+++ b/Engine/Source/Runtime/AnimationCore/Public/MultiplexStorage.h
@@ -20,9 +20,7 @@ struct ANIMATIONCORE_API FMultiplexAddress
 		, ByteIndex(INDEX_NONE)
 		, ElementSize(0)
 		, ElementCount(0)
-#if WITH_EDITORONLY_DATA
 		, Name(NAME_None)
-#endif
 	{
 	}
 
@@ -37,10 +35,8 @@ struct ANIMATIONCORE_API FMultiplexAddress
 	UPROPERTY()
 	int32 ElementCount;
 
-#if WITH_EDITORONLY_DATA
 	UPROPERTY()
 		FName Name;
-#endif
 
 	FORCEINLINE bool IsArray() const { return ElementCount > 1; }
 	FORCEINLINE int32 NumBytes() const { return ElementCount * ElementSize;  }
@@ -89,18 +85,17 @@ struct ANIMATIONCORE_API FMultiplexStorage
 	
 public:
 
-	FMultiplexStorage();
+	FMultiplexStorage(bool bInUseNames = true);
 	~FMultiplexStorage();
 
 	FMultiplexStorage& operator= (const FMultiplexStorage &InOther);
 
+	FORCEINLINE bool SupportsNames() const { return bUseNameMap;  }
 	FORCEINLINE int32 Num() const { return Addresses.Num(); }
 	FORCEINLINE const FMultiplexAddress& operator[](int32 InIndex) const { return Addresses[InIndex]; }
 	FORCEINLINE FMultiplexAddress& operator[](int32 InIndex) { return Addresses[InIndex]; }
-#if WITH_EDITORONLY_DATA
 	FORCEINLINE const FMultiplexAddress& operator[](const FName& InName) const { return Addresses[GetIndex(InName)]; }
 	FORCEINLINE FMultiplexAddress& operator[](const FName& InName) { return Addresses[GetIndex(InName)]; }
-#endif
 
 	FORCEINLINE TArray<FMultiplexAddress>::RangedForIteratorType      begin() { return Addresses.begin(); }
 	FORCEINLINE TArray<FMultiplexAddress>::RangedForConstIteratorType begin() const { return Addresses.begin(); }
@@ -146,10 +141,27 @@ public:
 		return TArrayView<T>((T*)&Data[Address.ByteIndex], Address.ElementCount);
 	}
 
-#if WITH_EDITORONLY_DATA
+	bool Copy(
+		int32 InSourceAddressIndex,
+		int32 InTargetAddressIndex,
+		int32 InSourceByteOffset = INDEX_NONE,
+		int32 InTargetByteOffset = INDEX_NONE,
+		int32 InNumBytes = INDEX_NONE);
+
+	bool Copy(
+		const FName& InSourceName,
+		const FName& InTargetName,
+		int32 InSourceByteOffset = INDEX_NONE,
+		int32 InTargetByteOffset = INDEX_NONE,
+		int32 InNumBytes = INDEX_NONE);
 
 	FORCEINLINE int32 GetIndex(const FName& InName) const
 	{
+		if (!bUseNameMap)
+		{
+			return INDEX_NONE;
+		}
+
 		if (NameMap.Num() != Addresses.Num())
 		{
 			for (int32 Index = 0; Index < Addresses.Num(); Index++)
@@ -172,18 +184,39 @@ public:
 		return INDEX_NONE;
 	}
 
-	FORCEINLINE bool IsNameAvailable(const FName& InPotentialNewName) const { return GetIndex(InPotentialNewName) == INDEX_NONE; }
+	FORCEINLINE bool IsNameAvailable(const FName& InPotentialNewName) const\
+	{
+		if (!bUseNameMap)
+		{
+			return false;
+		}
+		return GetIndex(InPotentialNewName) == INDEX_NONE;
+	}
 
 	void Reset();
 
+	int32 Add(int32 InElementSize, int32 InCount, const void* InData = nullptr);
+
 	int32 Add(const FName& InNewName, int32 InElementSize, int32 InCount, const void* InData = nullptr);
 
+	template<class T>
+	int32 Add(const T& InValue)
+	{
+		return Add(sizeof(T), 1, (const void*)&InValue);
+	}
+
 	template<class T>
 	int32 Add(const FName& InNewName, const T& InValue)
 	{
 		return Add(InNewName, sizeof(T), 1, (const void*)&InValue);
 	}
 
+	template<class T>
+	int32 AddArray(const TArray<T>& InArray)
+	{
+		return Add(sizeof(T), InArray.Num(), (const void*)InArray.GetData());
+	}
+
 	template<class T>
 	int32 AddArray(const FName& InNewName, const TArray<T>& InArray)
 	{
@@ -197,8 +230,6 @@ public:
 	bool Resize(int32 InAddressIndex, int32 InNewElementCount);
 	bool Resize(const FName& InAddressName, int32 InNewElementCount);
 
-#endif
-
 	void UpdateAddresses();
 
 private:
@@ -206,14 +237,15 @@ private:
 	// disable copy constructor
 	FMultiplexStorage(const FMultiplexStorage& Other) {}
 
+	UPROPERTY()
+	bool bUseNameMap;
+
 	UPROPERTY()
 	TArray<FMultiplexAddress> Addresses;
 
 	UPROPERTY()
 	TArray<uint8> Data;
 
-#if WITH_EDITORONLY_DATA
 	UPROPERTY(transient)
 	TMap<FName, int32> NameMap;
-#endif
 };
