Multiplex ByteCode data structure + unit tests. This is the baseline byte code management data structure for the multiplex vm.

diff --git a/Engine/Source/Runtime/AnimationCore/Private/MultiplexByteCode.cpp b/Engine/Source/Runtime/AnimationCore/Private/MultiplexByteCode.cpp
new file mode 100644
--- /dev/null
+++ b/Engine/Source/Runtime/AnimationCore/Private/MultiplexByteCode.cpp
@@ -0,0 +1,76 @@
+// Copyright 1998-2019 Epic Games, Inc. All Rights Reserved.
+
+#include "MultiplexByteCode.h"
+
+FMultiplexByteCode::FMultiplexByteCode()
+{
+}
+
+void FMultiplexByteCode::Reset()
+{
+	ByteCode.Reset();
+}
+
+uint64 FMultiplexByteCode::AddCopyOp(const FMultiplexArgument& InSource, const FMultiplexArgument& InTarget, int32 InSourceOffset, int32 InTargetOffset, int32 InNumBytes)
+{
+	FMultiplexCopyOp Op(InSource, InTarget, InSourceOffset, InTargetOffset, InNumBytes);
+	return AddOp(Op);
+}
+
+uint64 FMultiplexByteCode::AddIncrementOp(const FMultiplexArgument& InArg)
+{
+	FMultiplexIncrementOp Op(InArg);
+	return AddOp(Op);
+}
+
+uint64 FMultiplexByteCode::AddDecrementOp(const FMultiplexArgument& InArg)
+{
+	FMultiplexDecrementOp Op(InArg);
+	return AddOp(Op);
+}
+
+uint64 FMultiplexByteCode::AddEqualsOp(const FMultiplexArgument& InA, const FMultiplexArgument& InB, const FMultiplexArgument& InResult)
+{
+	FMultiplexEqualsOp Op(InA, InB, InResult);
+	return AddOp(Op);
+}
+
+uint64 FMultiplexByteCode::AddNotEqualsOp(const FMultiplexArgument& InA, const FMultiplexArgument& InB, const FMultiplexArgument& InResult)
+{
+	FMultiplexNotEqualsOp Op(InA, InB, InResult);
+	return AddOp(Op);
+}
+
+uint64 FMultiplexByteCode::AddJumpOp(uint64 InByteCodeIndex)
+{
+	FMultiplexJumpOp Op(InByteCodeIndex);
+	return AddOp(Op);
+}
+
+uint64 FMultiplexByteCode::AddJumpIfTrueOp(uint64 InByteCodeIndex, const FMultiplexArgument& InCondition)
+{
+	FMultiplexJumpIfTrueOp Op(InByteCodeIndex, InCondition);
+	return AddOp(Op);
+}
+
+uint64 FMultiplexByteCode::AddJumpIfFalseOp(uint64 InByteCodeIndex, const FMultiplexArgument& InCondition)
+{
+	FMultiplexJumpIfFalseOp Op(InByteCodeIndex, InCondition);
+	return AddOp(Op);
+}
+
+uint64 FMultiplexByteCode::AddExecuteOp(uint16 InFunctionIndex, const TArrayView<FMultiplexArgument>& InArguments)
+{
+	FMultiplexExecuteOp Op(InFunctionIndex, (uint16)InArguments.Num());
+	uint64 OpAddress = AddOp(Op);
+	uint64 ArgumentsAddress = (uint64)ByteCode.AddUninitialized(sizeof(FMultiplexArgument) * InArguments.Num());
+	FMemory::Memcpy(ByteCode.GetData() + ArgumentsAddress, InArguments.GetData(), sizeof(FMultiplexArgument) * InArguments.Num());
+	return OpAddress;
+}
+
+uint64 FMultiplexByteCode::AddExitOp()
+{
+	FMultiplexExitOp Op;
+	return AddOp(Op);
+}
+
diff --git a/Engine/Source/Runtime/AnimationCore/Public/MultiplexByteCode.h b/Engine/Source/Runtime/AnimationCore/Public/MultiplexByteCode.h
new file mode 100644
--- /dev/null
+++ b/Engine/Source/Runtime/AnimationCore/Public/MultiplexByteCode.h
@@ -0,0 +1,310 @@
+// Copyright 1998-2019 Epic Games, Inc. All Rights Reserved.
+
+#pragma once
+
+#include "CoreMinimal.h"
+#include "MultiplexRegistry.h"
+#include "MultiplexByteCode.generated.h"
+
+UENUM()
+enum class EMultiplexOpCode : uint8
+{
+	Copy,
+	Increment,
+	Decrement,
+	Equals,
+	NotEquals,
+	Jump,
+	JumpIfTrue,
+	JumpIfFalse,
+	Execute,
+	Exit,
+	Invalid
+};
+
+struct ANIMATIONCORE_API FMultiplexBaseOp
+{
+	FMultiplexBaseOp(EMultiplexOpCode InOpCode = EMultiplexOpCode::Invalid)
+	: OpCode(InOpCode)
+	{
+	}
+
+	EMultiplexOpCode OpCode;
+};
+
+struct ANIMATIONCORE_API FMultiplexCopyOp : public FMultiplexBaseOp
+{
+	FMultiplexCopyOp()
+	: FMultiplexBaseOp(EMultiplexOpCode::Copy)
+	, Source()
+	, Target()
+	, SourceOffset(INDEX_NONE)
+	, TargetOffset(INDEX_NONE)
+	, NumBytes(INDEX_NONE)
+	{
+	}
+
+	FMultiplexCopyOp(
+		FMultiplexArgument InSource,
+		FMultiplexArgument InTarget,
+		int32 InSourceOffset = INDEX_NONE,
+		int32 InTargetOffset = INDEX_NONE,
+		int32 InNumBytes = INDEX_NONE
+	)
+		: FMultiplexBaseOp(EMultiplexOpCode::Copy)
+		, Source(InSource)
+		, Target(InTarget)
+		, SourceOffset(InSourceOffset)
+		, TargetOffset(InTargetOffset)
+		, NumBytes(InNumBytes)
+	{
+	}
+
+	FMultiplexArgument Source;
+	FMultiplexArgument Target;
+	int32 SourceOffset;
+	int32 TargetOffset;
+	int32 NumBytes;
+};
+
+struct ANIMATIONCORE_API FMultiplexIncrementOp : public FMultiplexBaseOp
+{
+	FMultiplexIncrementOp()
+	: FMultiplexBaseOp(EMultiplexOpCode::Increment)
+	, Arg()
+	{
+	}
+
+	FMultiplexIncrementOp(FMultiplexArgument InArg)
+	: FMultiplexBaseOp(EMultiplexOpCode::Increment)
+	, Arg(InArg)
+	{
+	}
+
+	FMultiplexArgument Arg;
+};
+
+struct ANIMATIONCORE_API FMultiplexDecrementOp : public FMultiplexBaseOp
+{
+	FMultiplexDecrementOp()
+	: FMultiplexBaseOp(EMultiplexOpCode::Decrement)
+	, Arg()
+	{
+	}
+
+	FMultiplexDecrementOp(FMultiplexArgument InArg)
+	: FMultiplexBaseOp(EMultiplexOpCode::Decrement)
+	, Arg(InArg)
+	{
+	}
+
+	FMultiplexArgument Arg;
+};
+
+struct ANIMATIONCORE_API FMultiplexEqualsOp : public FMultiplexBaseOp
+{
+	FMultiplexEqualsOp()
+	: FMultiplexBaseOp(EMultiplexOpCode::Equals)
+	, A()
+	, B()
+	, Result()
+	{
+	}
+
+	FMultiplexEqualsOp(
+		FMultiplexArgument InA,
+		FMultiplexArgument InB,
+		FMultiplexArgument InResult
+	)
+	: FMultiplexBaseOp(EMultiplexOpCode::Equals)
+	, A(InA)
+	, B(InB)
+	, Result(InResult)
+	{
+	}
+
+	FMultiplexArgument A;
+	FMultiplexArgument B;
+	FMultiplexArgument Result;
+};
+
+struct ANIMATIONCORE_API FMultiplexNotEqualsOp : public FMultiplexBaseOp
+{
+	FMultiplexNotEqualsOp()
+	: FMultiplexBaseOp(EMultiplexOpCode::NotEquals)
+	, A()
+	, B()
+	, Result()
+	{
+	}
+
+	FMultiplexNotEqualsOp(
+		FMultiplexArgument InA,
+		FMultiplexArgument InB,
+		FMultiplexArgument InResult
+	)
+	: FMultiplexBaseOp(EMultiplexOpCode::NotEquals)
+	, A(InA)
+	, B(InB)
+	, Result(InResult)
+	{
+	}
+
+	FMultiplexArgument A;
+	FMultiplexArgument B;
+	FMultiplexArgument Result;
+};
+
+struct ANIMATIONCORE_API FMultiplexJumpOp : public FMultiplexBaseOp
+{
+	FMultiplexJumpOp()
+	: FMultiplexBaseOp(EMultiplexOpCode::Jump)
+	, ByteCodeIndex(UINT64_MAX)
+	{
+	}
+
+	FMultiplexJumpOp(uint64 InByteCodeIndex)
+	: FMultiplexBaseOp(EMultiplexOpCode::Jump)
+	, ByteCodeIndex(InByteCodeIndex)
+	{
+	}
+
+	uint64 ByteCodeIndex;
+};
+
+struct ANIMATIONCORE_API FMultiplexJumpIfTrueOp : public FMultiplexBaseOp
+{
+	FMultiplexJumpIfTrueOp()
+	: FMultiplexBaseOp(EMultiplexOpCode::JumpIfTrue)
+	, ByteCodeIndex(UINT64_MAX)
+	, Condition()
+	{
+	}
+
+	FMultiplexJumpIfTrueOp(uint64 InByteCodeIndex, FMultiplexArgument InCondition)
+	: FMultiplexBaseOp(EMultiplexOpCode::JumpIfTrue)
+	, ByteCodeIndex(InByteCodeIndex)
+	, Condition(InCondition)
+	{
+	}
+
+	uint64 ByteCodeIndex;
+	FMultiplexArgument Condition;
+};
+
+struct ANIMATIONCORE_API FMultiplexJumpIfFalseOp : public FMultiplexBaseOp
+{
+	FMultiplexJumpIfFalseOp()
+	: FMultiplexBaseOp(EMultiplexOpCode::JumpIfFalse)
+	, ByteCodeIndex(UINT64_MAX)
+	, Condition()
+	{
+	}
+
+	FMultiplexJumpIfFalseOp(uint64 InByteCodeIndex, FMultiplexArgument InCondition)
+	: FMultiplexBaseOp(EMultiplexOpCode::JumpIfFalse)
+	, ByteCodeIndex(InByteCodeIndex)
+	, Condition(InCondition)
+	{
+	}
+
+	uint64 ByteCodeIndex;
+	FMultiplexArgument Condition;
+};
+
+struct ANIMATIONCORE_API FMultiplexExecuteOp : public FMultiplexBaseOp
+{
+	FMultiplexExecuteOp()
+	: FMultiplexBaseOp(EMultiplexOpCode::Execute)
+	, FunctionIndex(INDEX_NONE)
+	, ArgumentCount(0)
+	{
+	}
+
+	FMultiplexExecuteOp(uint16 InFunctionIndex, uint16 InArgumentCount)
+	: FMultiplexBaseOp(EMultiplexOpCode::Execute)
+	, FunctionIndex(InFunctionIndex)
+	, ArgumentCount(InArgumentCount)
+	{
+	}
+
+	uint16 FunctionIndex;
+	uint16 ArgumentCount;
+};
+
+
+struct ANIMATIONCORE_API FMultiplexExitOp : public FMultiplexBaseOp
+{
+	FMultiplexExitOp()
+	: FMultiplexBaseOp(EMultiplexOpCode::Exit)
+	{
+	}
+};
+
+USTRUCT()
+struct ANIMATIONCORE_API FMultiplexByteCode
+{
+	GENERATED_USTRUCT_BODY()
+
+public:
+
+	FMultiplexByteCode();
+
+	void Reset();
+
+	uint64 AddCopyOp(const FMultiplexArgument& InSource, const FMultiplexArgument& InTarget, int32 InSourceOffset = INDEX_NONE, int32 InTargetOffset = INDEX_NONE, int32 InNumBytes = INDEX_NONE);
+	uint64 AddIncrementOp(const FMultiplexArgument& InArg);
+	uint64 AddDecrementOp(const FMultiplexArgument& InArg);
+	uint64 AddEqualsOp(const FMultiplexArgument& InA, const FMultiplexArgument& InB, const FMultiplexArgument& InResult);
+	uint64 AddNotEqualsOp(const FMultiplexArgument& InA, const FMultiplexArgument& InB, const FMultiplexArgument& InResult);
+	uint64 AddJumpOp(uint64 InByteCodeIndex);
+	uint64 AddJumpIfTrueOp(uint64 InByteCodeIndex, const FMultiplexArgument& InCondition);
+	uint64 AddJumpIfFalseOp(uint64 InByteCodeIndex, const FMultiplexArgument& InCondition);
+	uint64 AddExecuteOp(uint16 InFunctionIndex, const TArrayView<FMultiplexArgument>& InArguments);
+	uint64 AddExitOp();
+
+	FORCEINLINE EMultiplexOpCode GetOpCodeAt(uint64 InByteCodeIndex) const
+	{
+		ensure(InByteCodeIndex >= 0 && InByteCodeIndex < ByteCode.Num());
+		return (EMultiplexOpCode)ByteCode[InByteCodeIndex];
+	}
+
+	template<class OpType>
+	FORCEINLINE const OpType& GetOpAt(uint64 InByteCodeIndex) const
+	{
+		ensure(InByteCodeIndex >= 0 && InByteCodeIndex <= ByteCode.Num() - sizeof(OpType));
+		return *(const OpType*)(ByteCode.GetData() + InByteCodeIndex);
+	}
+
+	FORCEINLINE TArrayView<FMultiplexArgument> GetArgumentsAt(uint64 InByteCodeIndex, uint16 InArgumentCount) const
+	{
+		ensure(InByteCodeIndex >= 0 && InByteCodeIndex <= ByteCode.Num() - sizeof(FMultiplexArgument) * InArgumentCount);
+		return TArrayView<FMultiplexArgument>((FMultiplexArgument*)(ByteCode.GetData() + InByteCodeIndex), InArgumentCount);
+	}
+
+	FORCEINLINE TArrayView<FMultiplexArgument> GetArgumentsForExecuteOp(uint64 InByteCodeIndex) const
+	{
+		const FMultiplexExecuteOp& ExecuteOp = GetOpAt<FMultiplexExecuteOp>(InByteCodeIndex);
+		return GetArgumentsAt(InByteCodeIndex + sizeof(FMultiplexExecuteOp), ExecuteOp.ArgumentCount);
+	}
+
+	FORCEINLINE const TArrayView<uint8> GetByteCode() const
+	{
+		const uint8* Data = ByteCode.GetData();
+		return TArrayView<uint8>((uint8*)Data, ByteCode.Num());
+	}
+
+private:
+
+	template<class OpType>
+	FORCEINLINE uint64 AddOp(const OpType& InOp)
+	{
+		uint64 Address = (uint64)ByteCode.AddUninitialized(sizeof(OpType));
+		FMemory::Memcpy(ByteCode.GetData() + Address, &InOp, sizeof(OpType));
+		return Address;
+	}
+
+	// storage for all functions
+	UPROPERTY()
+	TArray<uint8> ByteCode;
+};
diff --git a/Engine/Source/Runtime/AnimationCore/Public/MultiplexRegistry.h b/Engine/Source/Runtime/AnimationCore/Public/MultiplexRegistry.h
--- a/Engine/Source/Runtime/AnimationCore/Public/MultiplexRegistry.h
+++ b/Engine/Source/Runtime/AnimationCore/Public/MultiplexRegistry.h
@@ -19,6 +19,8 @@ public:
 	{
 	}
 
+	FORCEINLINE bool operator ==(int32 InOther) const { return Address == InOther; }
+
 	FORCEINLINE bool IsLiteral() const { return Address < 0; }
 	FORCEINLINE int32 StorageType() const { return Address < 0 ? 1 : 0; }
 	FORCEINLINE int32 Index() const { return Address < 0 ? -(Address + 1) : Address; }
