Multiplex ByteCode data structure, table iterator + unit test

diff --git a/Engine/Source/Runtime/AnimationCore/Private/MultiplexByteCode.cpp b/Engine/Source/Runtime/AnimationCore/Private/MultiplexByteCode.cpp
--- a/Engine/Source/Runtime/AnimationCore/Private/MultiplexByteCode.cpp
+++ b/Engine/Source/Runtime/AnimationCore/Private/MultiplexByteCode.cpp
@@ -2,6 +2,21 @@
 
 #include "MultiplexByteCode.h"
 
+FMultiplexByteCodeTable::FMultiplexByteCodeTable()
+{
+}
+
+FMultiplexByteCodeTable::FMultiplexByteCodeTable(const FMultiplexByteCode& InByteCode)
+{
+	uint64 Address = 0;
+	while (Address < InByteCode.Num())
+	{
+		EMultiplexOpCode OpCode = InByteCode.GetOpCodeAt(Address);
+		Entries.Add(FMultiplexByteCodeTableEntry(OpCode, Address));
+		Address += InByteCode.GetOpNumBytesAt(Address);
+	}
+}
+
 FMultiplexByteCode::FMultiplexByteCode()
 {
 }
@@ -11,6 +26,71 @@ void FMultiplexByteCode::Reset()
 	ByteCode.Reset();
 }
 
+uint64 FMultiplexByteCode::Num() const
+{
+	return (uint64)ByteCode.Num();
+}
+
+uint64 FMultiplexByteCode::GetOpNumBytesAt(uint64 InByteCodeIndex, bool bIncludeArguments) const
+{
+	EMultiplexOpCode OpCode = GetOpCodeAt(InByteCodeIndex);
+	switch (OpCode)
+	{
+		case EMultiplexOpCode::Copy:
+		{
+			return (uint64)sizeof(FMultiplexCopyOp);
+		}
+		case EMultiplexOpCode::Increment:
+		{
+			return (uint64)sizeof(FMultiplexIncrementOp);
+		}
+		case EMultiplexOpCode::Decrement:
+		{
+			return (uint64)sizeof(FMultiplexDecrementOp);
+		}
+		case EMultiplexOpCode::Equals:
+		{
+			return (uint64)sizeof(FMultiplexEqualsOp);
+		}
+		case EMultiplexOpCode::NotEquals:
+		{
+			return (uint64)sizeof(FMultiplexNotEqualsOp);
+		}
+		case EMultiplexOpCode::Jump:
+		{
+			return (uint64)sizeof(FMultiplexJumpOp);
+		}
+		case EMultiplexOpCode::JumpIfTrue:
+		{
+			return (uint64)sizeof(FMultiplexJumpIfTrueOp);
+		}
+		case EMultiplexOpCode::JumpIfFalse:
+		{
+			return (uint64)sizeof(FMultiplexJumpIfFalseOp);
+		}
+		case EMultiplexOpCode::Execute:
+		{
+			uint64 NumBytes = (uint64)sizeof(FMultiplexExecuteOp);
+			if(bIncludeArguments)
+			{
+				const FMultiplexExecuteOp& ExecuteOp = GetOpAt<FMultiplexExecuteOp>(InByteCodeIndex);
+				NumBytes += (uint64)ExecuteOp.ArgumentCount * (uint64)sizeof(FMultiplexArgument);
+			}
+			return NumBytes;
+		}
+		case EMultiplexOpCode::Exit:
+		{
+			return (uint64)sizeof(FMultiplexExitOp);
+		}
+		case EMultiplexOpCode::Invalid:
+		{
+			ensure(false);
+			return 0;
+		}
+	}
+	return 0;
+}
+
 uint64 FMultiplexByteCode::AddCopyOp(const FMultiplexArgument& InSource, const FMultiplexArgument& InTarget, int32 InSourceOffset, int32 InTargetOffset, int32 InNumBytes)
 {
 	FMultiplexCopyOp Op(InSource, InTarget, InSourceOffset, InTargetOffset, InNumBytes);
diff --git a/Engine/Source/Runtime/AnimationCore/Public/MultiplexByteCode.h b/Engine/Source/Runtime/AnimationCore/Public/MultiplexByteCode.h
--- a/Engine/Source/Runtime/AnimationCore/Public/MultiplexByteCode.h
+++ b/Engine/Source/Runtime/AnimationCore/Public/MultiplexByteCode.h
@@ -6,6 +6,8 @@
 #include "MultiplexRegistry.h"
 #include "MultiplexByteCode.generated.h"
 
+struct FMultiplexByteCode;
+
 UENUM()
 enum class EMultiplexOpCode : uint8
 {
@@ -232,7 +234,6 @@ struct ANIMATIONCORE_API FMultiplexExecuteOp : public FMultiplexBaseOp
 	uint16 ArgumentCount;
 };
 
-
 struct ANIMATIONCORE_API FMultiplexExitOp : public FMultiplexBaseOp
 {
 	FMultiplexExitOp()
@@ -241,6 +242,47 @@ struct ANIMATIONCORE_API FMultiplexExitOp : public FMultiplexBaseOp
 	}
 };
 
+USTRUCT()
+struct ANIMATIONCORE_API FMultiplexByteCodeTableEntry
+{
+	GENERATED_USTRUCT_BODY()
+
+	FMultiplexByteCodeTableEntry(EMultiplexOpCode InOpCode = EMultiplexOpCode::Invalid, uint64 InByteCodeIndex = UINT64_MAX)
+		: OpCode(InOpCode)
+		, ByteCodeIndex(InByteCodeIndex)
+	{
+	}
+
+	UPROPERTY()
+	EMultiplexOpCode OpCode;
+
+	UPROPERTY()
+	uint64 ByteCodeIndex;
+};
+
+USTRUCT()
+struct ANIMATIONCORE_API FMultiplexByteCodeTable
+{
+	GENERATED_USTRUCT_BODY()
+
+public:
+
+	FMultiplexByteCodeTable();
+
+	FORCEINLINE int32 Num() const { return Entries.Num(); }
+	FORCEINLINE const FMultiplexByteCodeTableEntry& operator[](int32 InIndex) const { return Entries[InIndex]; }
+
+private:
+
+	// hide utility constructor
+	FMultiplexByteCodeTable(const FMultiplexByteCode& InByteCode);
+
+	UPROPERTY()
+	TArray<FMultiplexByteCodeTableEntry> Entries;
+
+	friend struct FMultiplexByteCode;
+};
+
 USTRUCT()
 struct ANIMATIONCORE_API FMultiplexByteCode
 {
@@ -251,6 +293,7 @@ public:
 	FMultiplexByteCode();
 
 	void Reset();
+	uint64 Num() const;
 
 	uint64 AddCopyOp(const FMultiplexArgument& InSource, const FMultiplexArgument& InTarget, int32 InSourceOffset = INDEX_NONE, int32 InTargetOffset = INDEX_NONE, int32 InNumBytes = INDEX_NONE);
 	uint64 AddIncrementOp(const FMultiplexArgument& InArg);
@@ -263,12 +306,19 @@ public:
 	uint64 AddExecuteOp(uint16 InFunctionIndex, const TArrayView<FMultiplexArgument>& InArguments);
 	uint64 AddExitOp();
 
+	FORCEINLINE FMultiplexByteCodeTable GetTable() const
+	{
+		return FMultiplexByteCodeTable(*this);
+	}
+
 	FORCEINLINE EMultiplexOpCode GetOpCodeAt(uint64 InByteCodeIndex) const
 	{
 		ensure(InByteCodeIndex >= 0 && InByteCodeIndex < ByteCode.Num());
 		return (EMultiplexOpCode)ByteCode[InByteCodeIndex];
 	}
 
+	uint64 GetOpNumBytesAt(uint64 InByteCodeIndex, bool bIncludeArguments = true) const;
+
 	template<class OpType>
 	FORCEINLINE const OpType& GetOpAt(uint64 InByteCodeIndex) const
 	{
@@ -276,6 +326,13 @@ public:
 		return *(const OpType*)(ByteCode.GetData() + InByteCodeIndex);
 	}
 
+	template<class OpType>
+	FORCEINLINE const OpType& GetOpAt(const FMultiplexByteCodeTableEntry& InEntry) const
+	{
+		ensure(OpType().OpCode == InEntry.OpCode);
+		return GetOpAt<OpType>(InEntry.ByteCodeIndex);
+	}
+
 	FORCEINLINE TArrayView<FMultiplexArgument> GetArgumentsAt(uint64 InByteCodeIndex, uint16 InArgumentCount) const
 	{
 		ensure(InByteCodeIndex >= 0 && InByteCodeIndex <= ByteCode.Num() - sizeof(FMultiplexArgument) * InArgumentCount);
@@ -288,6 +345,12 @@ public:
 		return GetArgumentsAt(InByteCodeIndex + sizeof(FMultiplexExecuteOp), ExecuteOp.ArgumentCount);
 	}
 
+	FORCEINLINE TArrayView<FMultiplexArgument> GetArgumentsForExecuteOp(const FMultiplexByteCodeTableEntry& InEntry) const
+	{
+		const FMultiplexExecuteOp& ExecuteOp = GetOpAt<FMultiplexExecuteOp>(InEntry);
+		return GetArgumentsAt(InEntry.ByteCodeIndex + sizeof(FMultiplexExecuteOp), ExecuteOp.ArgumentCount);
+	}
+
 	FORCEINLINE const TArrayView<uint8> GetByteCode() const
 	{
 		const uint8* Data = ByteCode.GetData();
