Multiplex: Added functionality for proper struct initialization and destruction within the storage

diff --git a/Engine/Source/Runtime/AnimationCore/Private/MultiplexRegistry.cpp b/Engine/Source/Runtime/AnimationCore/Private/MultiplexRegistry.cpp
--- a/Engine/Source/Runtime/AnimationCore/Private/MultiplexRegistry.cpp
+++ b/Engine/Source/Runtime/AnimationCore/Private/MultiplexRegistry.cpp
@@ -12,7 +12,6 @@ FMultiplexRegistry& FMultiplexRegistry::Get()
 
 void FMultiplexRegistry::Refresh()
 {
-	// todo
 }
 
 void FMultiplexRegistry::Register(const TCHAR* InName, FMultiplexFunctionPtr InFunctionPtr)
diff --git a/Engine/Source/Runtime/AnimationCore/Private/MultiplexStorage.cpp b/Engine/Source/Runtime/AnimationCore/Private/MultiplexStorage.cpp
--- a/Engine/Source/Runtime/AnimationCore/Private/MultiplexStorage.cpp
+++ b/Engine/Source/Runtime/AnimationCore/Private/MultiplexStorage.cpp
@@ -9,23 +9,72 @@ FMultiplexStorage::FMultiplexStorage(bool bInUseNames)
 
 FMultiplexStorage::~FMultiplexStorage()
 {
+	Reset();
 }
 
 FMultiplexStorage& FMultiplexStorage::operator= (const FMultiplexStorage &InOther)
 {
+	Reset();
+
 	bUseNameMap = InOther.bUseNameMap;
-	Data.Reset();
 	Data.Append(InOther.Data);
-	Addresses.Reset();
 	Addresses.Append(InOther.Addresses);
+	ScriptStructs.Append(InOther.ScriptStructs);
+
 	UpdateAddresses();
+
+	for (int32 Index = 0; Index < Addresses.Num(); Index++)
+	{
+		Construct(Index);
+
+		switch (Addresses[Index].Type)
+		{
+			case EMultiplexAddressType::Plain:
+			{
+				// the mem copy through the data array is enough!
+				break;
+			}
+			case EMultiplexAddressType::String:
+			{
+				FString* TargetStrings = (FString*)GetData(Index);
+				FString* SourceStrings = (FString*)InOther.GetData(Index);
+				for (int32 ElementIndex = 0; ElementIndex < Addresses[Index].ElementCount; ElementIndex++)
+				{
+					TargetStrings[ElementIndex] = SourceStrings[ElementIndex];
+				}
+				break;
+			}
+			case EMultiplexAddressType::Name:
+			{
+				// is this right? Nothing to do since name can just be mem copied?
+				break;
+			}
+			case EMultiplexAddressType::Struct:
+			{
+				UScriptStruct* ScriptStruct = GetScriptStruct(Index);
+				ScriptStruct->CopyScriptStruct(GetData(Index), InOther.GetData(Index), Addresses[Index].ElementCount);
+				break;
+			}
+			case EMultiplexAddressType::Invalid:
+			{
+				break;
+			}
+		}
+	}
+
 	return *this;
 }
 
 void FMultiplexStorage::Reset()
 {
+	for (int32 Index = 0; Index < Addresses.Num(); Index++)
+	{
+		Destroy(Index);
+	}
+
 	Data.Reset();
 	Addresses.Reset();
+	ScriptStructs.Reset();
 	NameMap.Reset();
 }
 
@@ -76,7 +125,42 @@ bool FMultiplexStorage::Copy(
 		return false;
 	}
 
-	FMemory::Memcpy(&Data[TargetStartByte], &Data[SourceStartByte], TargetNumBytes);
+	switch (Target.Type)
+	{
+		case EMultiplexAddressType::Plain:
+		{
+			FMemory::Memcpy(&Data[TargetStartByte], &Data[SourceStartByte], TargetNumBytes);
+			break;
+		}
+		case EMultiplexAddressType::Struct:
+		{
+			UScriptStruct* ScriptStruct = GetScriptStruct(InTargetAddressIndex);
+			int32 NumStructs = TargetNumBytes / ScriptStruct->GetStructureSize();
+			ScriptStruct->CopyScriptStruct(&Data[TargetStartByte], &Data[SourceStartByte], NumStructs);
+			break;
+		}
+		case EMultiplexAddressType::Name:
+		{
+			// is this right? can we just copy a name?
+			FMemory::Memcpy(&Data[TargetStartByte], &Data[SourceStartByte], TargetNumBytes);
+			break;
+		}
+		case EMultiplexAddressType::String:
+		{
+			int32 NumStrings = TargetNumBytes / sizeof(FString);
+			FString* TargetStrings = (FString*)&Data[TargetStartByte];
+			FString* SourceStrings = (FString*)&Data[SourceStartByte];
+			for (int32 Index = 0; Index < NumStrings; Index++)
+			{
+				TargetStrings[Index] = SourceStrings[Index];
+			}
+			break;
+		}
+		case EMultiplexAddressType::Invalid:
+		{
+			return false;
+		}
+	}
 
 	return true;
 }
@@ -101,35 +185,24 @@ bool FMultiplexStorage::Copy(
 	return Copy(SourceAddressIndex, TargetAddressIndex, InSourceByteOffset, InTargetByteOffset, InNumBytes);
 }
 
-int32 FMultiplexStorage::Add(int32 InElementSize, int32 InCount, const void* InData)
+int32 FMultiplexStorage::Allocate(const FName& InNewName, int32 InElementSize, int32 InCount, const void* InDataPtr)
 {
-	FName Name = NAME_None;
-
-	if (bUseNameMap)
+	FName Name = InNewName;
+	if (bUseNameMap && InNewName == NAME_None)
 	{
 		const TCHAR* AddressPrefix = TEXT("Address");
 		int32 AddressSuffix = 0;
 		do
 		{
 			Name = FName(*FString::Printf(TEXT("%s_%d"), AddressPrefix, AddressSuffix++));
-		}
-		while (!IsNameAvailable(Name));
+		} while (!IsNameAvailable(Name));
 	}
 
-	return Add(Name, InElementSize, InCount, InData);
-}
-
-int32 FMultiplexStorage::Add(const FName& InNewName, int32 InElementSize, int32 InCount, const void* InData)
-{
-	if (bUseNameMap)
-	{
-		ensure(InNewName != NAME_None);
-	}
 	ensure(InElementSize > 0 && InCount > 0);
 
 	if (bUseNameMap)
 	{
-		if (!IsNameAvailable(InNewName))
+		if (!IsNameAvailable(Name))
 		{
 			return INDEX_NONE;
 		}
@@ -139,16 +212,17 @@ int32 FMultiplexStorage::Add(const FName& InNewName, int32 InElementSize, int32
 	NewAddress.ByteIndex = Data.Num();
 	if (bUseNameMap)
 	{
-		NewAddress.Name = InNewName;
+		NewAddress.Name = Name;
 	}
 	NewAddress.ElementSize = InElementSize;
 	NewAddress.ElementCount = InCount;
+	NewAddress.Type = EMultiplexAddressType::Plain;
 
 	Data.AddZeroed(NewAddress.NumBytes());
 
-	if (InData != nullptr)
+	if (InDataPtr != nullptr)
 	{
-		FMemory::Memcpy(&Data[NewAddress.ByteIndex], InData, NewAddress.NumBytes());
+		FMemory::Memcpy(&Data[NewAddress.ByteIndex], InDataPtr, NewAddress.NumBytes());
 	}
 
 	int32 AddressIndex = Addresses.Num();
@@ -157,6 +231,107 @@ int32 FMultiplexStorage::Add(const FName& InNewName, int32 InElementSize, int32
 	return AddressIndex;
 }
 
+int32 FMultiplexStorage::Allocate(int32 InElementSize, int32 InCount, const void* InDataPtr)
+{
+	return Allocate(NAME_None, InElementSize, InCount, InDataPtr);
+}
+
+bool FMultiplexStorage::Construct(int32 InAddressIndex, int32 InElementIndex)
+{
+	ensure(Addresses.IsValidIndex(InAddressIndex));
+
+	const FMultiplexAddress& Address = Addresses[InAddressIndex];
+	switch (Address.Type)
+	{
+		case EMultiplexAddressType::Struct:
+		{
+			void* DataPtr = (void*)&Data[InElementIndex == INDEX_NONE ? Address.ByteIndex : Address.ByteIndex + InElementIndex * Address.ElementSize];
+			int32 Count = InElementIndex == INDEX_NONE ? Address.ElementCount : 1;
+
+			UScriptStruct* ScriptStruct = GetScriptStruct(InAddressIndex);
+			ScriptStruct->InitializeStruct(DataPtr, Count);
+			break;
+		}
+		case EMultiplexAddressType::String:
+		{
+			FString* DataPtr = (FString*)&Data[InElementIndex == INDEX_NONE ? Address.ByteIndex : Address.ByteIndex + InElementIndex * Address.ElementSize];
+			int32 Count = InElementIndex == INDEX_NONE ? Address.ElementCount : 1;
+
+			FMemory::Memzero(DataPtr, Address.NumBytes());
+			for (int32 Index = 0; Index < Count; Index++)
+			{
+				DataPtr[Index] = FString();
+			}
+			break;
+		}
+		case EMultiplexAddressType::Name:
+		{
+			FName* DataPtr = (FName*)&Data[InElementIndex == INDEX_NONE ? Address.ByteIndex : Address.ByteIndex + InElementIndex * Address.ElementSize];
+			int32 Count = InElementIndex == INDEX_NONE ? Address.ElementCount : 1;
+
+			FMemory::Memzero(DataPtr, Address.NumBytes());
+			for (int32 Index = 0; Index < Count; Index++)
+			{
+				DataPtr[Index] = FName();
+			}
+			break;
+		}
+		default:
+		{
+			return false;
+		}
+	}
+
+	return true;
+}
+
+bool FMultiplexStorage::Destroy(int32 InAddressIndex, int32 InElementIndex)
+{
+	ensure(Addresses.IsValidIndex(InAddressIndex));
+
+	const FMultiplexAddress& Address = Addresses[InAddressIndex];
+	switch (Address.Type)
+	{
+		case EMultiplexAddressType::Struct:
+		{
+			void* DataPtr = (void*)&Data[InElementIndex == INDEX_NONE ? Address.ByteIndex : Address.ByteIndex + InElementIndex * Address.ElementSize];
+			int32 Count = InElementIndex == INDEX_NONE ? Address.ElementCount : 1;
+
+			UScriptStruct* ScriptStruct = GetScriptStruct(InAddressIndex);
+			ScriptStruct->DestroyStruct(DataPtr, Count);
+			break;
+		}
+		case EMultiplexAddressType::String:
+		{
+			FString* DataPtr = (FString*)&Data[InElementIndex == INDEX_NONE ? Address.ByteIndex : Address.ByteIndex + InElementIndex * Address.ElementSize];
+			int32 Count = InElementIndex == INDEX_NONE ? Address.ElementCount : 1;
+
+			for (int32 Index = 0; Index < Count; Index++)
+			{
+				DataPtr[Index].Empty();
+			}
+			break;
+		}
+		case EMultiplexAddressType::Name:
+		{
+			FName* DataPtr = (FName*)&Data[InElementIndex == INDEX_NONE ? Address.ByteIndex : Address.ByteIndex + InElementIndex * Address.ElementSize];
+			int32 Count = InElementIndex == INDEX_NONE ? Address.ElementCount : 1;
+
+			for (int32 Index = 0; Index < Count; Index++)
+			{
+				DataPtr[Index] = FName();
+			}
+			break;
+		}
+		default:
+		{
+			return false;
+		}
+	}
+
+	return true;
+}
+
 bool FMultiplexStorage::Remove(int32 InAddressIndex)
 {
 	if (InAddressIndex < 0 || InAddressIndex >= Addresses.Num())
@@ -164,6 +339,8 @@ bool FMultiplexStorage::Remove(int32 InAddressIndex)
 		return false;
 	}
 
+	Destroy(InAddressIndex);
+
 	FMultiplexAddress AddressToRemove = Addresses[InAddressIndex];
 	Data.RemoveAt(AddressToRemove.ByteIndex, AddressToRemove.NumBytes());
 	Addresses.RemoveAt(InAddressIndex);
@@ -226,14 +403,21 @@ bool FMultiplexStorage::Resize(int32 InAddressIndex, int32 InNewElementCount)
 		return false;
 	}
 
-	if (Addresses[InAddressIndex].ElementCount > InNewElementCount) // shrink
+	FMultiplexAddress& Address = Addresses[InAddressIndex];
+
+	if (Address.ElementCount > InNewElementCount) // shrink
 	{
-		int32 ElementsToRemove = Addresses[InAddressIndex].ElementCount - InNewElementCount;
-		int32 NumBytesToRemove = Addresses[InAddressIndex].ElementSize * ElementsToRemove;
-		int32 FirstByteToRemove = Addresses[InAddressIndex].ByteIndex + Addresses[InAddressIndex].ElementSize * InNewElementCount;
+		int32 ElementsToRemove = Address.ElementCount - InNewElementCount;
+		int32 NumBytesToRemove = Address.ElementSize * ElementsToRemove;
+		int32 FirstByteToRemove = Address.ByteIndex + Address.ElementSize * InNewElementCount;
+
+		for (int32 ElementIndex = InNewElementCount; ElementIndex < Address.ElementCount; ElementIndex++)
+		{
+			Destroy(InAddressIndex, ElementIndex);
+		}
 
 		Data.RemoveAt(FirstByteToRemove, NumBytesToRemove);
-		Addresses[InAddressIndex].ElementCount = InNewElementCount;
+		Address.ElementCount = InNewElementCount;
 
 		for (int32 AddressIndex = InAddressIndex + 1; AddressIndex < Addresses.Num(); AddressIndex++)
 		{
@@ -242,12 +426,18 @@ bool FMultiplexStorage::Resize(int32 InAddressIndex, int32 InNewElementCount)
 	}
 	else // grow
 	{
-		int32 ElementsToAdd = InNewElementCount - Addresses[InAddressIndex].ElementCount;
-		int32 NumBytesToAdd = Addresses[InAddressIndex].ElementSize * ElementsToAdd;
-		int32 FirstByteToAdd = Addresses[InAddressIndex].ByteIndex + Addresses[InAddressIndex].ElementSize * Addresses[InAddressIndex].ElementCount;
+		int32 OldElementCount = Address.ElementCount;
+		int32 ElementsToAdd = InNewElementCount - Address.ElementCount;
+		int32 NumBytesToAdd = Address.ElementSize * ElementsToAdd;
+		int32 FirstByteToAdd = Address.ByteIndex + Address.ElementSize * Address.ElementCount;
 
 		Data.InsertZeroed(FirstByteToAdd, NumBytesToAdd);
-		Addresses[InAddressIndex].ElementCount = InNewElementCount;
+		Address.ElementCount = InNewElementCount;
+
+		for (int32 ElementIndex = OldElementCount; ElementIndex < InNewElementCount; ElementIndex++)
+		{
+			Construct(InAddressIndex, ElementIndex);
+		}
 
 		for (int32 AddressIndex = InAddressIndex + 1; AddressIndex < Addresses.Num(); AddressIndex++)
 		{
@@ -288,3 +478,19 @@ void FMultiplexStorage::UpdateAddresses()
 		}
 	}
 }
+
+void FMultiplexStorage::FillWithZeroes(int32 InAddressIndex)
+{
+	ensure(Addresses.IsValidIndex(InAddressIndex));
+	FMemory::Memzero(GetData(InAddressIndex), Addresses[InAddressIndex].NumBytes());
+}
+
+int32 FMultiplexStorage::FindOrAddScriptStruct(UScriptStruct* InScriptStruct)
+{
+	int32 StructIndex = INDEX_NONE;
+	if (ScriptStructs.Find(InScriptStruct, StructIndex))
+	{
+		return StructIndex;
+	}
+	return ScriptStructs.Add(InScriptStruct);
+}
diff --git a/Engine/Source/Runtime/AnimationCore/Public/MultiplexStorage.h b/Engine/Source/Runtime/AnimationCore/Public/MultiplexStorage.h
--- a/Engine/Source/Runtime/AnimationCore/Public/MultiplexStorage.h
+++ b/Engine/Source/Runtime/AnimationCore/Public/MultiplexStorage.h
@@ -6,20 +6,35 @@
 #include "UObject/ObjectMacros.h"
 #include "MultiplexStorage.generated.h"
 
+UENUM()
+enum class EMultiplexAddressType : uint8
+{
+	Plain,
+	String,
+	Name,
+	Struct,
+	Invalid
+};
+
 USTRUCT()
 struct ANIMATIONCORE_API FMultiplexAddress
 {
 	GENERATED_USTRUCT_BODY()
 
 	FMultiplexAddress()
-		: Pointer(nullptr)
+		: Type(EMultiplexAddressType::Invalid)
+		, Pointer(nullptr)
 		, ByteIndex(INDEX_NONE)
 		, ElementSize(0)
 		, ElementCount(0)
 		, Name(NAME_None)
+		, ScriptStructIndex(INDEX_NONE)
 	{
 	}
 
+	UPROPERTY()
+	EMultiplexAddressType Type;
+
 	void* Pointer;
 
 	UPROPERTY()
@@ -34,13 +49,16 @@ struct ANIMATIONCORE_API FMultiplexAddress
 	UPROPERTY()
 	FName Name;
 
+	UPROPERTY()
+	int32 ScriptStructIndex;
+
 	FORCEINLINE bool IsArray() const { return ElementCount > 1; }
+	FORCEINLINE bool IsPlain() const { return ScriptStructIndex == INDEX_NONE;  }
 	FORCEINLINE int32 NumBytes() const { return ElementCount * ElementSize;  }
 
 	template<class T>
 	FORCEINLINE const T* Get() const
 	{
-		ensure(sizeof(T) == ElementSize);
 		ensure(ElementCount > 0);
 		return (const T*)Pointer;
 	}
@@ -54,7 +72,6 @@ struct ANIMATIONCORE_API FMultiplexAddress
 	template<class T>
 	FORCEINLINE T* Get()
 	{
-		ensure(sizeof(T) == ElementSize);
 		ensure(ElementCount > 0);
 		return (T*)Pointer;
 	}
@@ -68,7 +85,6 @@ struct ANIMATIONCORE_API FMultiplexAddress
 	template<class T>
 	FORCEINLINE TArrayView<T> GetArray()
 	{
-		ensure(sizeof(T) == ElementSize);
 		ensure(ElementCount > 0);
 		return TArrayView<T>((T*)Pointer, ElementCount);
 	}
@@ -98,11 +114,27 @@ public:
 	FORCEINLINE TArray<FMultiplexAddress>::RangedForIteratorType      end() { return Addresses.end(); }
 	FORCEINLINE TArray<FMultiplexAddress>::RangedForConstIteratorType end() const { return Addresses.end(); }
 
+	FORCEINLINE const void* GetData(int32 InAddressIndex) const
+	{
+		ensure(Addresses.IsValidIndex(InAddressIndex));
+		const FMultiplexAddress& Address = Addresses[InAddressIndex];
+		ensure(Address.ElementCount > 0);
+		return (const void*)&Data[Address.ByteIndex];
+	}
+
+	FORCEINLINE void* GetData(int32 InAddressIndex)
+	{
+		ensure(Addresses.IsValidIndex(InAddressIndex));
+		const FMultiplexAddress& Address = Addresses[InAddressIndex];
+		ensure(Address.ElementCount > 0);
+		return (void*)&Data[Address.ByteIndex];
+	}
+
 	template<class T>
 	FORCEINLINE const T* Get(int32 InAddressIndex) const
 	{
+		ensure(Addresses.IsValidIndex(InAddressIndex));
 		const FMultiplexAddress& Address = Addresses[InAddressIndex];
-		ensure(sizeof(T) == Address.ElementSize);
 		ensure(Address.ElementCount > 0);
 		return (const T*)&Data[Address.ByteIndex];
 	}
@@ -116,8 +148,8 @@ public:
 	template<class T>
 	FORCEINLINE T* Get(int32 InAddressIndex)
 	{
+		ensure(Addresses.IsValidIndex(InAddressIndex));
 		const FMultiplexAddress& Address = Addresses[InAddressIndex];
-		ensure(sizeof(T) == Address.ElementSize);
 		ensure(Address.ElementCount > 0);
 		return (T*)&Data[Address.ByteIndex];
 	}
@@ -131,12 +163,24 @@ public:
 	template<class T>
 	FORCEINLINE TArrayView<T> GetArray(int32 InAddressIndex)
 	{
+		ensure(Addresses.IsValidIndex(InAddressIndex));
 		const FMultiplexAddress& Address = Addresses[InAddressIndex];
-		ensure(sizeof(T) == Address.ElementSize);
 		ensure(Address.ElementCount > 0);
 		return TArrayView<T>((T*)&Data[Address.ByteIndex], Address.ElementCount);
 	}
 
+	FORCEINLINE UScriptStruct* GetScriptStruct(int32 InAddressIndex) const
+	{
+		ensure(Addresses.IsValidIndex(InAddressIndex));
+		const FMultiplexAddress& Address = Addresses[InAddressIndex];
+		if (Address.ScriptStructIndex != INDEX_NONE)
+		{
+			ensure(ScriptStructs.IsValidIndex(Address.ScriptStructIndex));
+			return ScriptStructs[Address.ScriptStructIndex];
+		}
+		return nullptr;
+	}
+
 	bool Copy(
 		int32 InSourceAddressIndex,
 		int32 InTargetAddressIndex,
@@ -191,32 +235,129 @@ public:
 
 	void Reset();
 
-	int32 Add(int32 InElementSize, int32 InCount, const void* InData = nullptr);
+	int32 Allocate(const FName& InNewName, int32 InElementSize, int32 InCount, const void* InDataPtr = nullptr);
+	int32 Allocate(int32 InElementSize, int32 InCount, const void* InDataPtr = nullptr);
+	bool Construct(int32 InAddressIndex, int32 InElementIndex = INDEX_NONE);
+	bool Destroy(int32 InAddressIndex, int32 InElementIndex = INDEX_NONE);
 
-	int32 Add(const FName& InNewName, int32 InElementSize, int32 InCount, const void* InData = nullptr);
+	int32 AddPlainArray(const FName& InNewName, int32 InElementSize, int32 InCount, const void* InDataPtr = nullptr)
+	{
+		return Allocate(InNewName, InElementSize, InCount, InDataPtr);
+	}
 
 	template<class T>
-	int32 Add(const T& InValue)
+	int32 AddPlainArray(const FName& InNewName, int32 InCount, const T* InDataPtr = nullptr)
 	{
-		return Add(sizeof(T), 1, (const void*)&InValue);
+		return Allocate(InNewName, sizeof(T), InCount, (const void*)InDataPtr);
 	}
 
 	template<class T>
-	int32 Add(const FName& InNewName, const T& InValue)
+	int32 AddPlainArray(const FName& InNewName, const TArray<T>& InArray)
 	{
-		return Add(InNewName, sizeof(T), 1, (const void*)&InValue);
+		return AddPlainArray<T>(InNewName, InArray.Num(), InArray.GetData());
 	}
 
 	template<class T>
-	int32 AddArray(const TArray<T>& InArray)
+	int32 AddPlainArray(const TArray<T>& InArray)
 	{
-		return Add(sizeof(T), InArray.Num(), (const void*)InArray.GetData());
+		return AddPlainArray<T>(NAME_None, InArray);
+	}
+
+	int32 AddPlain(const FName& InNewName, int32 InElementSize, const void* InValuePtr)
+	{
+		return AddPlainArray(InNewName, InElementSize, 1, InValuePtr);
+	}
+
+	int32 AddPlain(int32 InElementSize, const void* InValuePtr)
+	{
+		return AddPlain(NAME_None, InElementSize, InValuePtr);
 	}
 
 	template<class T>
-	int32 AddArray(const FName& InNewName, const TArray<T>& InArray)
+	int32 AddPlain(const FName& InNewName, const T& InValue)
 	{
-		return Add(InNewName, sizeof(T), InArray.Num(), (const void*)InArray.GetData());
+		return AddPlainArray<T>(InNewName, 1, &InValue);
+	}
+
+	template<class T>
+	int32 AddPlain(const T& InValue)
+	{
+		return AddPlain<T>(NAME_None, InValue);
+	}
+
+	int32 AddStructArray(const FName& InNewName, UScriptStruct* InScriptStruct, int32 InCount, const void* InDataPtr = nullptr)
+	{
+		int32 Address = Allocate(InNewName, InScriptStruct->GetStructureSize(), InCount, nullptr);
+		if (Address == INDEX_NONE)
+		{
+			return INDEX_NONE;
+		}
+
+		Addresses[Address].Type = EMultiplexAddressType::Struct;
+		Addresses[Address].ScriptStructIndex = FindOrAddScriptStruct(InScriptStruct);
+
+		// construct the content
+		Construct(Address);
+
+		// copy values from the provided data
+		if (InDataPtr != nullptr)
+		{
+			InScriptStruct->CopyScriptStruct(GetData(Address), InDataPtr, InCount);
+		}
+
+		return Address;
+	}
+
+	int32 AddStructArray(UScriptStruct* InScriptStruct, int32 InCount, const void* InDataPtr = nullptr)
+	{
+		return AddStructArray(NAME_None, InScriptStruct, InCount, InDataPtr);
+	}
+
+	template<class T>
+	int32 AddStructArray(const FName& InNewName, int32 InCount, const T* InDataPtr = nullptr)
+	{
+		// if you are hitting this - you might need to use AddPlainArray instead!
+		UScriptStruct* Struct = T::StaticStruct();
+		if (Struct == nullptr)
+		{
+			return INDEX_NONE;
+		}
+	
+		return AddStructArray(InNewName, Struct, InCount, InDataPtr);
+	}
+
+	template<class T>
+	int32 AddStructArray(const FName& InNewName, const TArray<T>& InArray)
+	{
+		return AddStructArray<T>(InNewName, InArray.Num(), InArray.GetData());
+	}
+
+	template<class T>
+	int32 AddStructArray(const TArray<T>& InArray)
+	{
+		return AddStructArray<T>(NAME_None, InArray);
+	}
+
+	int32 AddStruct(const FName& InNewName, UScriptStruct* InScriptStruct, const void* InValuePtr)
+	{
+		return AddStructArray(InNewName, InScriptStruct, 1, InValuePtr);
+	}
+
+	int32 AddStruct(UScriptStruct* InScriptStruct, const void* InValuePtr)
+	{
+		return AddStruct(NAME_None, InScriptStruct, InValuePtr);
+	}
+
+	template<class T>
+	int32 AddStruct(const FName& InNewName, const T& InValue)
+	{
+		return AddStructArray<T>(InNewName, 1, &InValue);
+	}
+
+	template<class T>
+	int32 AddStruct(const T& InValue)
+	{
+		return AddStruct<T>(NAME_None, InValue);
 	}
 
 	bool Remove(int32 InAddressIndex);
@@ -233,6 +374,9 @@ private:
 	// disable copy constructor
 	FMultiplexStorage(const FMultiplexStorage& Other) {}
 
+	void FillWithZeroes(int32 InAddressIndex);
+	int32 FindOrAddScriptStruct(UScriptStruct* InScriptStruct);
+
 	UPROPERTY()
 	bool bUseNameMap;
 
@@ -242,6 +386,9 @@ private:
 	UPROPERTY()
 	TArray<uint8> Data;
 
+	UPROPERTY()
+	TArray<UScriptStruct*> ScriptStructs;
+
 	UPROPERTY(transient)
 	TMap<FName, int32> NameMap;
 };
