Multiplex: Support for strings and names + unit tests

diff --git a/Engine/Source/Runtime/AnimationCore/Private/MultiplexStorage.cpp b/Engine/Source/Runtime/AnimationCore/Private/MultiplexStorage.cpp
--- a/Engine/Source/Runtime/AnimationCore/Private/MultiplexStorage.cpp
+++ b/Engine/Source/Runtime/AnimationCore/Private/MultiplexStorage.cpp
@@ -26,40 +26,7 @@ FMultiplexStorage& FMultiplexStorage::operator= (const FMultiplexStorage &InOthe
 	for (int32 Index = 0; Index < Addresses.Num(); Index++)
 	{
 		Construct(Index);
-
-		switch (Addresses[Index].Type)
-		{
-			case EMultiplexAddressType::Plain:
-			{
-				// the mem copy through the data array is enough!
-				break;
-			}
-			case EMultiplexAddressType::String:
-			{
-				FString* TargetStrings = (FString*)GetData(Index);
-				FString* SourceStrings = (FString*)InOther.GetData(Index);
-				for (int32 ElementIndex = 0; ElementIndex < Addresses[Index].ElementCount; ElementIndex++)
-				{
-					TargetStrings[ElementIndex] = SourceStrings[ElementIndex];
-				}
-				break;
-			}
-			case EMultiplexAddressType::Name:
-			{
-				// is this right? Nothing to do since name can just be mem copied?
-				break;
-			}
-			case EMultiplexAddressType::Struct:
-			{
-				UScriptStruct* ScriptStruct = GetScriptStruct(Index);
-				ScriptStruct->CopyScriptStruct(GetData(Index), InOther.GetData(Index), Addresses[Index].ElementCount);
-				break;
-			}
-			case EMultiplexAddressType::Invalid:
-			{
-				break;
-			}
-		}
+		Copy(Index, Index, &InOther);
 	}
 
 	return *this;
@@ -81,19 +48,25 @@ void FMultiplexStorage::Reset()
 bool FMultiplexStorage::Copy(
 	int32 InSourceAddressIndex,
 	int32 InTargetAddressIndex,
+	const FMultiplexStorage* InSourceStorage,
 	int32 InSourceByteOffset,
 	int32 InTargetByteOffset,
 	int32 InNumBytes)
 {
-	ensure(Addresses.IsValidIndex(InSourceAddressIndex));
+	if (InSourceStorage == nullptr)
+	{
+		InSourceStorage = this;
+	}
+
+	ensure(InSourceStorage->Addresses.IsValidIndex(InSourceAddressIndex));
 	ensure(Addresses.IsValidIndex(InTargetAddressIndex));
 
-	if (InSourceAddressIndex == InTargetAddressIndex && InSourceByteOffset == InTargetByteOffset)
+	if (InSourceAddressIndex == InTargetAddressIndex && InSourceByteOffset == InTargetByteOffset && this == InSourceStorage)
 	{
 		return false;
 	}
 
-	const FMultiplexAddress& Source = Addresses[InSourceAddressIndex];
+	const FMultiplexAddress& Source = InSourceStorage->Addresses[InSourceAddressIndex];
 	const FMultiplexAddress& Target = Addresses[InTargetAddressIndex];
 
 	int32 SourceStartByte = Source.ByteIndex;
@@ -129,27 +102,32 @@ bool FMultiplexStorage::Copy(
 	{
 		case EMultiplexAddressType::Plain:
 		{
-			FMemory::Memcpy(&Data[TargetStartByte], &Data[SourceStartByte], TargetNumBytes);
+			FMemory::Memcpy(&Data[TargetStartByte], &InSourceStorage->Data[SourceStartByte], TargetNumBytes);
 			break;
 		}
 		case EMultiplexAddressType::Struct:
 		{
 			UScriptStruct* ScriptStruct = GetScriptStruct(InTargetAddressIndex);
 			int32 NumStructs = TargetNumBytes / ScriptStruct->GetStructureSize();
-			ScriptStruct->CopyScriptStruct(&Data[TargetStartByte], &Data[SourceStartByte], NumStructs);
+			ScriptStruct->CopyScriptStruct(&Data[TargetStartByte], &InSourceStorage->Data[SourceStartByte], NumStructs);
 			break;
 		}
 		case EMultiplexAddressType::Name:
 		{
-			// is this right? can we just copy a name?
-			FMemory::Memcpy(&Data[TargetStartByte], &Data[SourceStartByte], TargetNumBytes);
+			int32 NumNames = TargetNumBytes / sizeof(FName);
+			TArrayView<FName> TargetNames((FName*)&Data[TargetStartByte], NumNames);
+			TArrayView<FName> SourceNames((FName*)&InSourceStorage->Data[SourceStartByte], NumNames);
+			for (int32 Index = 0; Index < NumNames; Index++)
+			{
+				TargetNames[Index] = SourceNames[Index];
+			}
 			break;
 		}
 		case EMultiplexAddressType::String:
 		{
 			int32 NumStrings = TargetNumBytes / sizeof(FString);
-			FString* TargetStrings = (FString*)&Data[TargetStartByte];
-			FString* SourceStrings = (FString*)&Data[SourceStartByte];
+			TArrayView<FString> TargetStrings((FString*)&Data[TargetStartByte], NumStrings);
+			TArrayView<FString> SourceStrings((FString*)&InSourceStorage->Data[SourceStartByte], NumStrings);
 			for (int32 Index = 0; Index < NumStrings; Index++)
 			{
 				TargetStrings[Index] = SourceStrings[Index];
@@ -168,6 +146,7 @@ bool FMultiplexStorage::Copy(
 bool FMultiplexStorage::Copy(
 	const FName& InSourceName,
 	const FName& InTargetName,
+	const FMultiplexStorage* InSourceStorage,
 	int32 InSourceByteOffset,
 	int32 InTargetByteOffset,
 	int32 InNumBytes)
@@ -182,7 +161,7 @@ bool FMultiplexStorage::Copy(
 		return false;
 	}
 
-	return Copy(SourceAddressIndex, TargetAddressIndex, InSourceByteOffset, InTargetByteOffset, InNumBytes);
+	return Copy(SourceAddressIndex, TargetAddressIndex, InSourceStorage, InSourceByteOffset, InTargetByteOffset, InNumBytes);
 }
 
 int32 FMultiplexStorage::Allocate(const FName& InNewName, int32 InElementSize, int32 InCount, const void* InDataPtr)
@@ -257,7 +236,7 @@ bool FMultiplexStorage::Construct(int32 InAddressIndex, int32 InElementIndex)
 			FString* DataPtr = (FString*)&Data[InElementIndex == INDEX_NONE ? Address.ByteIndex : Address.ByteIndex + InElementIndex * Address.ElementSize];
 			int32 Count = InElementIndex == INDEX_NONE ? Address.ElementCount : 1;
 
-			FMemory::Memzero(DataPtr, Address.NumBytes());
+			FMemory::Memzero(DataPtr, Count * Address.ElementSize);
 			for (int32 Index = 0; Index < Count; Index++)
 			{
 				DataPtr[Index] = FString();
@@ -269,7 +248,7 @@ bool FMultiplexStorage::Construct(int32 InAddressIndex, int32 InElementIndex)
 			FName* DataPtr = (FName*)&Data[InElementIndex == INDEX_NONE ? Address.ByteIndex : Address.ByteIndex + InElementIndex * Address.ElementSize];
 			int32 Count = InElementIndex == INDEX_NONE ? Address.ElementCount : 1;
 
-			FMemory::Memzero(DataPtr, Address.NumBytes());
+			FMemory::Memzero(DataPtr, Count * Address.ElementSize);
 			for (int32 Index = 0; Index < Count; Index++)
 			{
 				DataPtr[Index] = FName();
diff --git a/Engine/Source/Runtime/AnimationCore/Public/MultiplexStorage.h b/Engine/Source/Runtime/AnimationCore/Public/MultiplexStorage.h
--- a/Engine/Source/Runtime/AnimationCore/Public/MultiplexStorage.h
+++ b/Engine/Source/Runtime/AnimationCore/Public/MultiplexStorage.h
@@ -184,6 +184,7 @@ public:
 	bool Copy(
 		int32 InSourceAddressIndex,
 		int32 InTargetAddressIndex,
+		const FMultiplexStorage* InSourceStorage = nullptr,
 		int32 InSourceByteOffset = INDEX_NONE,
 		int32 InTargetByteOffset = INDEX_NONE,
 		int32 InNumBytes = INDEX_NONE);
@@ -191,6 +192,7 @@ public:
 	bool Copy(
 		const FName& InSourceName,
 		const FName& InTargetName,
+		const FMultiplexStorage* InSourceStorage = nullptr,
 		int32 InSourceByteOffset = INDEX_NONE,
 		int32 InTargetByteOffset = INDEX_NONE,
 		int32 InNumBytes = INDEX_NONE);
@@ -240,52 +242,131 @@ public:
 	bool Construct(int32 InAddressIndex, int32 InElementIndex = INDEX_NONE);
 	bool Destroy(int32 InAddressIndex, int32 InElementIndex = INDEX_NONE);
 
-	int32 AddPlainArray(const FName& InNewName, int32 InElementSize, int32 InCount, const void* InDataPtr = nullptr)
+	FORCEINLINE int32 AddPlainArray(const FName& InNewName, int32 InElementSize, int32 InCount, const void* InDataPtr = nullptr)
 	{
 		return Allocate(InNewName, InElementSize, InCount, InDataPtr);
 	}
 
 	template<class T>
-	int32 AddPlainArray(const FName& InNewName, int32 InCount, const T* InDataPtr = nullptr)
+	FORCEINLINE int32 AddPlainArray(const FName& InNewName, int32 InCount, const T* InDataPtr = nullptr)
 	{
 		return Allocate(InNewName, sizeof(T), InCount, (const void*)InDataPtr);
 	}
 
 	template<class T>
-	int32 AddPlainArray(const FName& InNewName, const TArray<T>& InArray)
+	FORCEINLINE int32 AddPlainArray(const FName& InNewName, const TArray<T>& InArray)
 	{
 		return AddPlainArray<T>(InNewName, InArray.Num(), InArray.GetData());
 	}
 
 	template<class T>
-	int32 AddPlainArray(const TArray<T>& InArray)
+	FORCEINLINE int32 AddPlainArray(const TArray<T>& InArray)
 	{
 		return AddPlainArray<T>(NAME_None, InArray);
 	}
 
-	int32 AddPlain(const FName& InNewName, int32 InElementSize, const void* InValuePtr)
+	FORCEINLINE int32 AddPlain(const FName& InNewName, int32 InElementSize, const void* InValuePtr)
 	{
 		return AddPlainArray(InNewName, InElementSize, 1, InValuePtr);
 	}
 
-	int32 AddPlain(int32 InElementSize, const void* InValuePtr)
+	FORCEINLINE int32 AddPlain(int32 InElementSize, const void* InValuePtr)
 	{
 		return AddPlain(NAME_None, InElementSize, InValuePtr);
 	}
 
 	template<class T>
-	int32 AddPlain(const FName& InNewName, const T& InValue)
+	FORCEINLINE int32 AddPlain(const FName& InNewName, const T& InValue)
 	{
 		return AddPlainArray<T>(InNewName, 1, &InValue);
 	}
 
 	template<class T>
-	int32 AddPlain(const T& InValue)
+	FORCEINLINE int32 AddPlain(const T& InValue)
 	{
 		return AddPlain<T>(NAME_None, InValue);
 	}
 
-	int32 AddStructArray(const FName& InNewName, UScriptStruct* InScriptStruct, int32 InCount, const void* InDataPtr = nullptr)
+	FORCEINLINE int32 AddNameArray(const FName& InNewName, int32 InCount, const FName* InDataPtr = nullptr)
+	{
+		int32 Address = Allocate(InNewName, sizeof(FName), InCount, nullptr);
+		Addresses[Address].Type = EMultiplexAddressType::Name;
+
+		Construct(Address);
+
+		if(InDataPtr)
+		{
+			FName* DataPtr = (FName*)GetData(Address);
+			for (int32 Index = 0; Index < InCount; Index++)
+			{
+				DataPtr[Index] = InDataPtr[Index];
+			}
+		}
+
+		return Address;
+	}
+
+	FORCEINLINE int32 AddNameArray(const FName& InNewName, const TArray<FName>& InArray)
+	{
+		return AddNameArray(InNewName, InArray.Num(), InArray.GetData());
+	}
+
+	FORCEINLINE int32 AddNameArray(const TArray<FName>& InArray)
+	{
+		return AddNameArray(NAME_None, InArray);
+	}
+
+	FORCEINLINE int32 AddName(const FName& InNewName, const FName& InValue)
+	{
+		return AddNameArray(InNewName, 1, &InValue);
+	}
+
+	FORCEINLINE int32 AddName(const FName& InValue)
+	{
+		return AddName(NAME_None, InValue);
+	}
+
+	FORCEINLINE int32 AddStringArray(const FName& InNewName, int32 InCount, const FString* InDataPtr = nullptr)
+	{
+		int32 Address = Allocate(InNewName, sizeof(FString), InCount, nullptr);
+		Addresses[Address].Type = EMultiplexAddressType::String;
+
+		Construct(Address);
+
+		if(InDataPtr)
+		{
+			FString* DataPtr = (FString*)GetData(Address);
+			for (int32 Index = 0; Index < InCount; Index++)
+			{
+				DataPtr[Index] = InDataPtr[Index];
+			}
+		}
+
+		return Address;
+	}
+
+	FORCEINLINE int32 AddStringArray(const FName& InNewName, const TArray<FString>& InArray)
+	{
+		return AddStringArray(InNewName, InArray.Num(), InArray.GetData());
+	}
+
+	FORCEINLINE int32 AddStringArray(const TArray<FString>& InArray)
+	{
+		return AddStringArray(NAME_None, InArray);
+	}
+
+	FORCEINLINE int32 AddString(const FName& InNewName, const FString& InValue)
+	{
+		return AddStringArray(InNewName, 1, &InValue);
+	}
+
+	FORCEINLINE int32 AddString(const FString& InValue)
+	{
+		return AddString(NAME_None, InValue);
+	}
+
+
+	FORCEINLINE int32 AddStructArray(const FName& InNewName, UScriptStruct* InScriptStruct, int32 InCount, const void* InDataPtr = nullptr)
 	{
 		int32 Address = Allocate(InNewName, InScriptStruct->GetStructureSize(), InCount, nullptr);
 		if (Address == INDEX_NONE)
@@ -308,13 +389,13 @@ public:
 		return Address;
 	}
 
-	int32 AddStructArray(UScriptStruct* InScriptStruct, int32 InCount, const void* InDataPtr = nullptr)
+	FORCEINLINE int32 AddStructArray(UScriptStruct* InScriptStruct, int32 InCount, const void* InDataPtr = nullptr)
 	{
 		return AddStructArray(NAME_None, InScriptStruct, InCount, InDataPtr);
 	}
 
 	template<class T>
-	int32 AddStructArray(const FName& InNewName, int32 InCount, const T* InDataPtr = nullptr)
+	FORCEINLINE int32 AddStructArray(const FName& InNewName, int32 InCount, const T* InDataPtr = nullptr)
 	{
 		// if you are hitting this - you might need to use AddPlainArray instead!
 		UScriptStruct* Struct = T::StaticStruct();
@@ -327,35 +408,35 @@ public:
 	}
 
 	template<class T>
-	int32 AddStructArray(const FName& InNewName, const TArray<T>& InArray)
+	FORCEINLINE int32 AddStructArray(const FName& InNewName, const TArray<T>& InArray)
 	{
 		return AddStructArray<T>(InNewName, InArray.Num(), InArray.GetData());
 	}
 
 	template<class T>
-	int32 AddStructArray(const TArray<T>& InArray)
+	FORCEINLINE int32 AddStructArray(const TArray<T>& InArray)
 	{
 		return AddStructArray<T>(NAME_None, InArray);
 	}
 
-	int32 AddStruct(const FName& InNewName, UScriptStruct* InScriptStruct, const void* InValuePtr)
+	FORCEINLINE int32 AddStruct(const FName& InNewName, UScriptStruct* InScriptStruct, const void* InValuePtr)
 	{
 		return AddStructArray(InNewName, InScriptStruct, 1, InValuePtr);
 	}
 
-	int32 AddStruct(UScriptStruct* InScriptStruct, const void* InValuePtr)
+	FORCEINLINE int32 AddStruct(UScriptStruct* InScriptStruct, const void* InValuePtr)
 	{
 		return AddStruct(NAME_None, InScriptStruct, InValuePtr);
 	}
 
 	template<class T>
-	int32 AddStruct(const FName& InNewName, const T& InValue)
+	FORCEINLINE int32 AddStruct(const FName& InNewName, const T& InValue)
 	{
 		return AddStructArray<T>(InNewName, 1, &InValue);
 	}
 
 	template<class T>
-	int32 AddStruct(const T& InValue)
+	FORCEINLINE int32 AddStruct(const T& InValue)
 	{
 		return AddStruct<T>(NAME_None, InValue);
 	}
