Multiplex: Support for struct alignment and copy assignment + unit tests

diff --git a/Engine/Source/Runtime/AnimationCore/Private/MultiplexStorage.cpp b/Engine/Source/Runtime/AnimationCore/Private/MultiplexStorage.cpp
--- a/Engine/Source/Runtime/AnimationCore/Private/MultiplexStorage.cpp
+++ b/Engine/Source/Runtime/AnimationCore/Private/MultiplexStorage.cpp
@@ -7,6 +7,11 @@ FMultiplexStorage::FMultiplexStorage(bool bInUseNames)
 {
 }
 
+FMultiplexStorage::FMultiplexStorage(const FMultiplexStorage& Other)
+{
+	*this = Other;
+}
+
 FMultiplexStorage::~FMultiplexStorage()
 {
 	Reset();
@@ -69,25 +74,25 @@ bool FMultiplexStorage::Copy(
 	const FMultiplexAddress& Source = InSourceStorage->Addresses[InSourceAddressIndex];
 	const FMultiplexAddress& Target = Addresses[InTargetAddressIndex];
 
-	int32 SourceStartByte = Source.ByteIndex;
-	int32 SourceNumBytes = Source.NumBytes();
+	int32 SourceStartByte = Source.FirstByte();
+	int32 SourceNumBytes = Source.NumBytes(false);
 	if(InSourceByteOffset != INDEX_NONE)
 	{
 		ensure(InNumBytes != INDEX_NONE);
-		ensure(InSourceByteOffset >= 0 && InSourceByteOffset < Source.NumBytes());
-		ensure(InNumBytes > 0 && InSourceByteOffset + InNumBytes <= Source.NumBytes());
+		ensure(InSourceByteOffset >= 0 && InSourceByteOffset < Source.NumBytes(false));
+		ensure(InNumBytes > 0 && InSourceByteOffset + InNumBytes <= Source.NumBytes(false));
 
 		SourceStartByte += InSourceByteOffset;
 		SourceNumBytes = InNumBytes;
 	}
 
-	int32 TargetStartByte = Target.ByteIndex;
-	int32 TargetNumBytes = Target.NumBytes();
+	int32 TargetStartByte = Target.FirstByte();
+	int32 TargetNumBytes = Target.NumBytes(false);
 	if(InTargetByteOffset != INDEX_NONE)
 	{
 		ensure(InNumBytes != INDEX_NONE);
-		ensure(InTargetByteOffset >= 0 && InTargetByteOffset < Target.NumBytes());
-		ensure(InNumBytes > 0 && InTargetByteOffset + InNumBytes <= Target.NumBytes());
+		ensure(InTargetByteOffset >= 0 && InTargetByteOffset < Target.NumBytes(false));
+		ensure(InNumBytes > 0 && InTargetByteOffset + InNumBytes <= Target.NumBytes(false));
 
 		TargetStartByte += InTargetByteOffset;
 		TargetNumBytes = InNumBytes;
@@ -164,7 +169,7 @@ bool FMultiplexStorage::Copy(
 	return Copy(SourceAddressIndex, TargetAddressIndex, InSourceStorage, InSourceByteOffset, InTargetByteOffset, InNumBytes);
 }
 
-int32 FMultiplexStorage::Allocate(const FName& InNewName, int32 InElementSize, int32 InCount, const void* InDataPtr)
+int32 FMultiplexStorage::Allocate(const FName& InNewName, int32 InElementSize, int32 InCount, const void* InDataPtr, bool bUpdateAddresses)
 {
 	FName Name = InNewName;
 	if (bUseNameMap && InNewName == NAME_None)
@@ -201,18 +206,22 @@ int32 FMultiplexStorage::Allocate(const FName& InNewName, int32 InElementSize, i
 
 	if (InDataPtr != nullptr)
 	{
-		FMemory::Memcpy(&Data[NewAddress.ByteIndex], InDataPtr, NewAddress.NumBytes());
+		FMemory::Memcpy(&Data[NewAddress.FirstByte()], InDataPtr, NewAddress.NumBytes());
 	}
 
 	int32 AddressIndex = Addresses.Num();
 	Addresses.Add(NewAddress);
-	UpdateAddresses();
+
+	if (bUpdateAddresses)
+	{
+		UpdateAddresses();
+	}
 	return AddressIndex;
 }
 
-int32 FMultiplexStorage::Allocate(int32 InElementSize, int32 InCount, const void* InDataPtr)
+int32 FMultiplexStorage::Allocate(int32 InElementSize, int32 InCount, const void* InDataPtr, bool bUpdateAddresses)
 {
-	return Allocate(NAME_None, InElementSize, InCount, InDataPtr);
+	return Allocate(NAME_None, InElementSize, InCount, InDataPtr, bUpdateAddresses);
 }
 
 bool FMultiplexStorage::Construct(int32 InAddressIndex, int32 InElementIndex)
@@ -224,7 +233,7 @@ bool FMultiplexStorage::Construct(int32 InAddressIndex, int32 InElementIndex)
 	{
 		case EMultiplexAddressType::Struct:
 		{
-			void* DataPtr = (void*)&Data[InElementIndex == INDEX_NONE ? Address.ByteIndex : Address.ByteIndex + InElementIndex * Address.ElementSize];
+			void* DataPtr = (void*)&Data[InElementIndex == INDEX_NONE ? Address.FirstByte() : Address.FirstByte() + InElementIndex * Address.ElementSize];
 			int32 Count = InElementIndex == INDEX_NONE ? Address.ElementCount : 1;
 
 			UScriptStruct* ScriptStruct = GetScriptStruct(InAddressIndex);
@@ -233,7 +242,7 @@ bool FMultiplexStorage::Construct(int32 InAddressIndex, int32 InElementIndex)
 		}
 		case EMultiplexAddressType::String:
 		{
-			FString* DataPtr = (FString*)&Data[InElementIndex == INDEX_NONE ? Address.ByteIndex : Address.ByteIndex + InElementIndex * Address.ElementSize];
+			FString* DataPtr = (FString*)&Data[InElementIndex == INDEX_NONE ? Address.FirstByte() : Address.FirstByte() + InElementIndex * Address.ElementSize];
 			int32 Count = InElementIndex == INDEX_NONE ? Address.ElementCount : 1;
 
 			FMemory::Memzero(DataPtr, Count * Address.ElementSize);
@@ -245,7 +254,7 @@ bool FMultiplexStorage::Construct(int32 InAddressIndex, int32 InElementIndex)
 		}
 		case EMultiplexAddressType::Name:
 		{
-			FName* DataPtr = (FName*)&Data[InElementIndex == INDEX_NONE ? Address.ByteIndex : Address.ByteIndex + InElementIndex * Address.ElementSize];
+			FName* DataPtr = (FName*)&Data[InElementIndex == INDEX_NONE ? Address.FirstByte() : Address.FirstByte() + InElementIndex * Address.ElementSize];
 			int32 Count = InElementIndex == INDEX_NONE ? Address.ElementCount : 1;
 
 			FMemory::Memzero(DataPtr, Count * Address.ElementSize);
@@ -273,7 +282,7 @@ bool FMultiplexStorage::Destroy(int32 InAddressIndex, int32 InElementIndex)
 	{
 		case EMultiplexAddressType::Struct:
 		{
-			void* DataPtr = (void*)&Data[InElementIndex == INDEX_NONE ? Address.ByteIndex : Address.ByteIndex + InElementIndex * Address.ElementSize];
+			void* DataPtr = (void*)&Data[InElementIndex == INDEX_NONE ? Address.FirstByte() : Address.FirstByte() + InElementIndex * Address.ElementSize];
 			int32 Count = InElementIndex == INDEX_NONE ? Address.ElementCount : 1;
 
 			UScriptStruct* ScriptStruct = GetScriptStruct(InAddressIndex);
@@ -282,7 +291,7 @@ bool FMultiplexStorage::Destroy(int32 InAddressIndex, int32 InElementIndex)
 		}
 		case EMultiplexAddressType::String:
 		{
-			FString* DataPtr = (FString*)&Data[InElementIndex == INDEX_NONE ? Address.ByteIndex : Address.ByteIndex + InElementIndex * Address.ElementSize];
+			FString* DataPtr = (FString*)&Data[InElementIndex == INDEX_NONE ? Address.FirstByte() : Address.FirstByte() + InElementIndex * Address.ElementSize];
 			int32 Count = InElementIndex == INDEX_NONE ? Address.ElementCount : 1;
 
 			for (int32 Index = 0; Index < Count; Index++)
@@ -293,7 +302,7 @@ bool FMultiplexStorage::Destroy(int32 InAddressIndex, int32 InElementIndex)
 		}
 		case EMultiplexAddressType::Name:
 		{
-			FName* DataPtr = (FName*)&Data[InElementIndex == INDEX_NONE ? Address.ByteIndex : Address.ByteIndex + InElementIndex * Address.ElementSize];
+			FName* DataPtr = (FName*)&Data[InElementIndex == INDEX_NONE ? Address.FirstByte() : Address.FirstByte() + InElementIndex * Address.ElementSize];
 			int32 Count = InElementIndex == INDEX_NONE ? Address.ElementCount : 1;
 
 			for (int32 Index = 0; Index < Count; Index++)
@@ -388,7 +397,7 @@ bool FMultiplexStorage::Resize(int32 InAddressIndex, int32 InNewElementCount)
 	{
 		int32 ElementsToRemove = Address.ElementCount - InNewElementCount;
 		int32 NumBytesToRemove = Address.ElementSize * ElementsToRemove;
-		int32 FirstByteToRemove = Address.ByteIndex + Address.ElementSize * InNewElementCount;
+		int32 FirstByteToRemove = Address.FirstByte() + Address.ElementSize * InNewElementCount;
 
 		for (int32 ElementIndex = InNewElementCount; ElementIndex < Address.ElementCount; ElementIndex++)
 		{
@@ -408,7 +417,7 @@ bool FMultiplexStorage::Resize(int32 InAddressIndex, int32 InNewElementCount)
 		int32 OldElementCount = Address.ElementCount;
 		int32 ElementsToAdd = InNewElementCount - Address.ElementCount;
 		int32 NumBytesToAdd = Address.ElementSize * ElementsToAdd;
-		int32 FirstByteToAdd = Address.ByteIndex + Address.ElementSize * Address.ElementCount;
+		int32 FirstByteToAdd = Address.FirstByte() + Address.ElementSize * Address.ElementCount;
 
 		Data.InsertZeroed(FirstByteToAdd, NumBytesToAdd);
 		Address.ElementCount = InNewElementCount;
@@ -443,9 +452,48 @@ bool FMultiplexStorage::Resize(const FName& InAddressName, int32 InNewElementCou
 
 void FMultiplexStorage::UpdateAddresses()
 {
-	for (FMultiplexAddress& Address : Addresses)
+	int32 AlignmentShift = 0;
+	for (int32 AddressIndex = 0; AddressIndex < Addresses.Num(); AddressIndex++)
+	{
+		FMultiplexAddress& Address = Addresses[AddressIndex];
+		Address.ByteIndex += AlignmentShift;
+
+		UScriptStruct* ScriptStruct = GetScriptStruct(AddressIndex);
+		if (ScriptStruct != nullptr)
+		{
+			UScriptStruct::ICppStructOps* TheCppStructOps = ScriptStruct->GetCppStructOps();
+			if (TheCppStructOps != NULL)
+			{
+				if (!TheCppStructOps->HasZeroConstructor())
+				{
+					void* Pointer = GetData(AddressIndex);
+
+					if (Address.AlignmentBytes > 0)
+					{
+						if (!IsAligned(Pointer, TheCppStructOps->GetAlignment()))
+						{
+							Data.RemoveAt(Address.ByteIndex, Address.AlignmentBytes);
+							AlignmentShift -= Address.AlignmentBytes;
+							Address.AlignmentBytes = 0;
+							Pointer = GetData(AddressIndex);
+						}
+					}
+
+					while (!IsAligned(Pointer, TheCppStructOps->GetAlignment()))
+					{
+						Data.InsertZeroed(Address.ByteIndex, 1);
+						Address.AlignmentBytes++;
+						AlignmentShift++;
+						Pointer = GetData(AddressIndex);
+					}
+				}
+			}
+		}
+	}
+
+	for (int32 AddressIndex = 0; AddressIndex < Addresses.Num(); AddressIndex++)
 	{
-		Address.Pointer = &Data[Address.ByteIndex];
+		Addresses[AddressIndex].Pointer = GetData(AddressIndex);
 	}
 
 	if (bUseNameMap)
diff --git a/Engine/Source/Runtime/AnimationCore/Public/MultiplexStorage.h b/Engine/Source/Runtime/AnimationCore/Public/MultiplexStorage.h
--- a/Engine/Source/Runtime/AnimationCore/Public/MultiplexStorage.h
+++ b/Engine/Source/Runtime/AnimationCore/Public/MultiplexStorage.h
@@ -21,12 +21,13 @@ struct ANIMATIONCORE_API FMultiplexAddress
 {
 	GENERATED_USTRUCT_BODY()
 
-	FMultiplexAddress()
+		FMultiplexAddress()
 		: Type(EMultiplexAddressType::Invalid)
 		, Pointer(nullptr)
 		, ByteIndex(INDEX_NONE)
 		, ElementSize(0)
 		, ElementCount(0)
+		, AlignmentBytes(0)
 		, Name(NAME_None)
 		, ScriptStructIndex(INDEX_NONE)
 	{
@@ -46,15 +47,19 @@ struct ANIMATIONCORE_API FMultiplexAddress
 	UPROPERTY()
 	int32 ElementCount;
 
+	UPROPERTY()
+	int32 AlignmentBytes;
+
 	UPROPERTY()
 	FName Name;
 
 	UPROPERTY()
 	int32 ScriptStructIndex;
 
+	FORCEINLINE int32 FirstByte() const { return ByteIndex + AlignmentBytes; }
 	FORCEINLINE bool IsArray() const { return ElementCount > 1; }
 	FORCEINLINE bool IsPlain() const { return ScriptStructIndex == INDEX_NONE;  }
-	FORCEINLINE int32 NumBytes() const { return ElementCount * ElementSize;  }
+	FORCEINLINE int32 NumBytes(bool bIncludeAlignment = true) const { return ElementCount * ElementSize + (bIncludeAlignment ? AlignmentBytes : 0); }
 
 	template<class T>
 	FORCEINLINE const T* Get() const
@@ -98,6 +103,7 @@ struct ANIMATIONCORE_API FMultiplexStorage
 public:
 
 	FMultiplexStorage(bool bInUseNames = true);
+	FMultiplexStorage(const FMultiplexStorage& Other);
 	~FMultiplexStorage();
 
 	FMultiplexStorage& operator= (const FMultiplexStorage &InOther);
@@ -119,7 +125,7 @@ public:
 		ensure(Addresses.IsValidIndex(InAddressIndex));
 		const FMultiplexAddress& Address = Addresses[InAddressIndex];
 		ensure(Address.ElementCount > 0);
-		return (const void*)&Data[Address.ByteIndex];
+		return (const void*)&Data[Address.FirstByte()];
 	}
 
 	FORCEINLINE void* GetData(int32 InAddressIndex)
@@ -127,7 +133,7 @@ public:
 		ensure(Addresses.IsValidIndex(InAddressIndex));
 		const FMultiplexAddress& Address = Addresses[InAddressIndex];
 		ensure(Address.ElementCount > 0);
-		return (void*)&Data[Address.ByteIndex];
+		return (void*)&Data[Address.FirstByte()];
 	}
 
 	template<class T>
@@ -136,7 +142,7 @@ public:
 		ensure(Addresses.IsValidIndex(InAddressIndex));
 		const FMultiplexAddress& Address = Addresses[InAddressIndex];
 		ensure(Address.ElementCount > 0);
-		return (const T*)&Data[Address.ByteIndex];
+		return (const T*)&Data[Address.FirstByte()];
 	}
 
 	template<class T>
@@ -151,7 +157,7 @@ public:
 		ensure(Addresses.IsValidIndex(InAddressIndex));
 		const FMultiplexAddress& Address = Addresses[InAddressIndex];
 		ensure(Address.ElementCount > 0);
-		return (T*)&Data[Address.ByteIndex];
+		return (T*)&Data[Address.FirstByte()];
 	}
 
 	template<class T>
@@ -166,7 +172,7 @@ public:
 		ensure(Addresses.IsValidIndex(InAddressIndex));
 		const FMultiplexAddress& Address = Addresses[InAddressIndex];
 		ensure(Address.ElementCount > 0);
-		return TArrayView<T>((T*)&Data[Address.ByteIndex], Address.ElementCount);
+		return TArrayView<T>((T*)&Data[Address.FirstByte()], Address.ElementCount);
 	}
 
 	FORCEINLINE UScriptStruct* GetScriptStruct(int32 InAddressIndex) const
@@ -237,11 +243,6 @@ public:
 
 	void Reset();
 
-	int32 Allocate(const FName& InNewName, int32 InElementSize, int32 InCount, const void* InDataPtr = nullptr);
-	int32 Allocate(int32 InElementSize, int32 InCount, const void* InDataPtr = nullptr);
-	bool Construct(int32 InAddressIndex, int32 InElementIndex = INDEX_NONE);
-	bool Destroy(int32 InAddressIndex, int32 InElementIndex = INDEX_NONE);
-
 	FORCEINLINE int32 AddPlainArray(const FName& InNewName, int32 InElementSize, int32 InCount, const void* InDataPtr = nullptr)
 	{
 		return Allocate(InNewName, InElementSize, InCount, InDataPtr);
@@ -368,7 +369,7 @@ public:
 
 	FORCEINLINE int32 AddStructArray(const FName& InNewName, UScriptStruct* InScriptStruct, int32 InCount, const void* InDataPtr = nullptr)
 	{
-		int32 Address = Allocate(InNewName, InScriptStruct->GetStructureSize(), InCount, nullptr);
+		int32 Address = Allocate(InNewName, InScriptStruct->GetStructureSize(), InCount, nullptr, false);
 		if (Address == INDEX_NONE)
 		{
 			return INDEX_NONE;
@@ -377,6 +378,8 @@ public:
 		Addresses[Address].Type = EMultiplexAddressType::Struct;
 		Addresses[Address].ScriptStructIndex = FindOrAddScriptStruct(InScriptStruct);
 
+		UpdateAddresses();
+
 		// construct the content
 		Construct(Address);
 
@@ -452,8 +455,10 @@ public:
 
 private:
 
-	// disable copy constructor
-	FMultiplexStorage(const FMultiplexStorage& Other) {}
+	int32 Allocate(const FName& InNewName, int32 InElementSize, int32 InCount, const void* InDataPtr = nullptr, bool bUpdateAddresses = true);
+	int32 Allocate(int32 InElementSize, int32 InCount, const void* InDataPtr = nullptr, bool bUpdateAddresses = true);
+	bool Construct(int32 InAddressIndex, int32 InElementIndex = INDEX_NONE);
+	bool Destroy(int32 InAddressIndex, int32 InElementIndex = INDEX_NONE);
 
 	void FillWithZeroes(int32 InAddressIndex);
 	int32 FindOrAddScriptStruct(UScriptStruct* InScriptStruct);
