Multiplex: VM data structure starts to work - basic instructions implemented + unit tests

diff --git a/Engine/Source/Programs/UnrealHeaderTool/Private/CodeGenerator.cpp b/Engine/Source/Programs/UnrealHeaderTool/Private/CodeGenerator.cpp
--- a/Engine/Source/Programs/UnrealHeaderTool/Private/CodeGenerator.cpp
+++ b/Engine/Source/Programs/UnrealHeaderTool/Private/CodeGenerator.cpp
@@ -3400,7 +3400,7 @@ void FNativeClassHeaderGenerator::ExportGeneratedStructBodyMacros(FOutputDevice&
 
 			MemberDeclarations.Add(FString::Printf(TEXT("%s %s"), *ParameterCPPType, *Name));
 			FString AddressName = FString::Printf(TEXT("%s_Address"), *Name);
-			MemberMulticastProlog.Add(FString::Printf(TEXT("%s %s = MultiplexStorage[%s.StorageType()][%s.Index()].%s<%s>();"), *ExtractedMulticastCPPType, *Name, *AddressName, *AddressName, *MulticastGetter, *MultiCastCPPType));
+			MemberMulticastProlog.Add(FString::Printf(TEXT("%s %s = MultiplexStorage[%s.StorageType()]->%s<%s>(%s.Index());"), *ExtractedMulticastCPPType, *Name, *AddressName, *MulticastGetter, *MultiCastCPPType, *AddressName));
 		}
 
 		OutGeneratedHeaderText.Log(TEXT("\n"));
@@ -3431,7 +3431,7 @@ void FNativeClassHeaderGenerator::ExportGeneratedStructBodyMacros(FOutputDevice&
 			{
 				FString ParamsSuffix = Info.Params.IsEmpty() ? TEXT("") : FString::Printf(TEXT(", %s"), *Info.Params);
 				MultiplexMethodsDeclarations += FString::Printf(TEXT("\tstatic %s Static%s(\r\n\t\t%s%s);\r\n"), *Info.ReturnType, *Info.Name, *FString::Join(MemberDeclarations, TEXT(",\r\n\t\t")), *ParamsSuffix);
-				MultiplexMethodsDeclarations += FString::Printf(TEXT("\tstatic %s Multiplex%s(const TArrayView<FMultiplexArgument>& MultiplexAddresses, FMultiplexStorage* MultiplexStorage, const TArrayView<void*>& AdditionalArgs);\r\n"), *Info.ReturnType, *Info.Name);
+				MultiplexMethodsDeclarations += FString::Printf(TEXT("\tstatic %s Multiplex%s(const TArrayView<FMultiplexArgument>& MultiplexAddresses, FMultiplexStorage** MultiplexStorage, const TArrayView<void*>& AdditionalArgs);\r\n"), *Info.ReturnType, *Info.Name);
 			}
 		}
 
@@ -3731,7 +3731,7 @@ void FNativeClassHeaderGenerator::ExportGeneratedStructBodyMacros(FOutputDevice&
 				ParamDeclarationSuffix = FString::Printf(TEXT(", %s"), *Info.Params);
 			}
 
-			Out.Logf(TEXT("%s %s::Multiplex%s(const TArrayView<FMultiplexArgument>& MultiplexAddresses, FMultiplexStorage* MultiplexStorage, const TArrayView<void*>& AdditionalArgs)\r\n"), *Info.ReturnType, StructNameCPP, *Info.Name);
+			Out.Logf(TEXT("%s %s::Multiplex%s(const TArrayView<FMultiplexArgument>& MultiplexAddresses, FMultiplexStorage** MultiplexStorage, const TArrayView<void*>& AdditionalArgs)\r\n"), *Info.ReturnType, StructNameCPP, *Info.Name);
 			Out.Log(TEXT("{\r\n"));
 
 			if (ParameterExtractionProlog.Num() > 0)
diff --git a/Engine/Source/Runtime/AnimationCore/Private/MultiplexByteCode.cpp b/Engine/Source/Runtime/AnimationCore/Private/MultiplexByteCode.cpp
--- a/Engine/Source/Runtime/AnimationCore/Private/MultiplexByteCode.cpp
+++ b/Engine/Source/Runtime/AnimationCore/Private/MultiplexByteCode.cpp
@@ -17,6 +17,11 @@ FMultiplexByteCodeTable::FMultiplexByteCodeTable(const FMultiplexByteCode& InByt
 	}
 }
 
+void FMultiplexByteCodeTable::Reset()
+{
+	Entries.Reset();
+}
+
 FMultiplexByteCode::FMultiplexByteCode()
 {
 }
diff --git a/Engine/Source/Runtime/AnimationCore/Private/MultiplexVM.cpp b/Engine/Source/Runtime/AnimationCore/Private/MultiplexVM.cpp
new file mode 100644
--- /dev/null
+++ b/Engine/Source/Runtime/AnimationCore/Private/MultiplexVM.cpp
@@ -0,0 +1,152 @@
+// Copyright 1998-2019 Epic Games, Inc. All Rights Reserved.
+
+#include "MultiplexVM.h"
+
+void UMultiplexVM::Reset()
+{
+	Literals.Reset();
+	WorkState.Reset();
+	FunctionNames.Reset();
+	Functions.Reset();
+	ByteCode.Reset();
+	Instructions.Reset();
+}
+
+int32 UMultiplexVM::AddMultiplexFunction(UScriptStruct* InMultiplexStruct, const FName& InFunctionName)
+{
+	check(InMultiplexStruct);
+	FString FunctionKey = FString::Printf(TEXT("F%s::%s"), *InMultiplexStruct->GetName(), *InFunctionName.ToString());
+	int32 FunctionIndex = FunctionNames.Find(FunctionKey);
+	if (FunctionIndex != INDEX_NONE)
+	{
+		return FunctionIndex;
+	}
+
+	FMultiplexFunctionPtr Function = FMultiplexRegistry::Get().Find(*FunctionKey);
+	if (Function == nullptr)
+	{
+		return INDEX_NONE;
+	}
+
+	FunctionNames.Add(FunctionKey);
+	return Functions.Add(Function);
+}
+
+void UMultiplexVM::ResolveFunctionsIfRequired()
+{
+	if (Functions.Num() != FunctionNames.Num())
+	{
+		Functions.Reset();
+		Functions.SetNumZeroed(FunctionNames.Num());
+
+		for (int32 FunctionIndex = 0; FunctionIndex < FunctionNames.Num(); FunctionIndex)
+		{
+			Functions[FunctionIndex] = FMultiplexRegistry::Get().Find(*FunctionNames[FunctionIndex]);
+		}
+	}
+}
+
+void UMultiplexVM::RefreshInstructionsIfRequired()
+{
+	if (Instructions.Num() == 0)
+	{
+		Instructions = ByteCode.GetTable();
+	}
+}
+
+bool UMultiplexVM::Execute(FMultiplexStorage** Storage, TArrayView<void*> AdditionalArgs)
+{
+	ResolveFunctionsIfRequired();
+	RefreshInstructionsIfRequired();
+
+	if (Instructions.Num() == 0)
+	{
+		return true;
+	}
+
+	FMultiplexStorage* LocalStorage[] = { &WorkState, &Literals }; 
+	if (Storage == nullptr)
+	{
+		Storage = LocalStorage;
+	}
+
+	int32 InstructionIndex = 0;
+	while (Instructions[InstructionIndex].OpCode != EMultiplexOpCode::Exit)
+	{
+		switch (Instructions[InstructionIndex].OpCode)
+		{
+			case EMultiplexOpCode::Copy:
+			{
+				const FMultiplexCopyOp& Op = ByteCode.GetOpAt<FMultiplexCopyOp>(Instructions[InstructionIndex]);
+				Storage[0]->Copy(Op.Source.Index(), Op.Target.Index(), Storage[Op.Source.StorageType()], Op.SourceOffset, Op.TargetOffset, Op.NumBytes);
+				InstructionIndex++;
+				break;
+			}
+			case EMultiplexOpCode::Increment:
+			{
+				const FMultiplexIncrementOp& Op = ByteCode.GetOpAt<FMultiplexIncrementOp>(Instructions[InstructionIndex]);
+				Storage[0]->GetRef<int32>(Op.Arg.Index())++;
+				InstructionIndex++;
+				break;
+			}
+			case EMultiplexOpCode::Decrement:
+			{
+				const FMultiplexDecrementOp& Op = ByteCode.GetOpAt<FMultiplexDecrementOp>(Instructions[InstructionIndex]);
+				Storage[0]->GetRef<int32>(Op.Arg.Index())--;
+				InstructionIndex++;
+				break;
+			}
+			case EMultiplexOpCode::Equals:
+			{
+				const FMultiplexEqualsOp& Op = ByteCode.GetOpAt<FMultiplexEqualsOp>(Instructions[InstructionIndex]);
+				ensureMsgf(false, TEXT("TO BE IMPLEMENTED"));
+				break;
+			}
+			case EMultiplexOpCode::NotEquals:
+			{
+				const FMultiplexNotEqualsOp& Op = ByteCode.GetOpAt<FMultiplexNotEqualsOp>(Instructions[InstructionIndex]);
+				ensureMsgf(false, TEXT("TO BE IMPLEMENTED"));
+				break;
+			}
+			case EMultiplexOpCode::Jump:
+			{
+				const FMultiplexJumpOp& Op = ByteCode.GetOpAt<FMultiplexJumpOp>(Instructions[InstructionIndex]);
+				ensureMsgf(false, TEXT("TO BE IMPLEMENTED"));
+				break;
+			}
+			case EMultiplexOpCode::JumpIfTrue:
+			{
+				const FMultiplexJumpIfTrueOp& Op = ByteCode.GetOpAt<FMultiplexJumpIfTrueOp>(Instructions[InstructionIndex]);
+				ensureMsgf(false, TEXT("TO BE IMPLEMENTED"));
+				break;
+			}
+			case EMultiplexOpCode::JumpIfFalse:
+			{
+				const FMultiplexJumpIfFalseOp& Op = ByteCode.GetOpAt<FMultiplexJumpIfFalseOp>(Instructions[InstructionIndex]);
+				ensureMsgf(false, TEXT("TO BE IMPLEMENTED"));
+				break;
+			}
+			case EMultiplexOpCode::Execute:
+			{
+				const FMultiplexExecuteOp& Op = ByteCode.GetOpAt<FMultiplexExecuteOp>(Instructions[InstructionIndex]);
+				TArrayView<FMultiplexArgument> Args = ByteCode.GetArgumentsForExecuteOp(Instructions[InstructionIndex]);
+				(*Functions[Op.FunctionIndex])(Args, Storage, AdditionalArgs);
+				InstructionIndex++;
+				break;
+			}
+			case EMultiplexOpCode::Exit:
+			case EMultiplexOpCode::Invalid:
+			{
+				ensure(false);
+				return false;
+			}
+		}
+
+		if (!Instructions.IsValidIndex(InstructionIndex))
+		{
+			break;
+		}
+	}
+
+	return true;
+}
diff --git a/Engine/Source/Runtime/AnimationCore/Public/MultiplexByteCode.h b/Engine/Source/Runtime/AnimationCore/Public/MultiplexByteCode.h
--- a/Engine/Source/Runtime/AnimationCore/Public/MultiplexByteCode.h
+++ b/Engine/Source/Runtime/AnimationCore/Public/MultiplexByteCode.h
@@ -161,36 +161,36 @@ struct ANIMATIONCORE_API FMultiplexJumpOp : public FMultiplexBaseOp
 {
 	FMultiplexJumpOp()
 	: FMultiplexBaseOp(EMultiplexOpCode::Jump)
-	, ByteCodeIndex(UINT64_MAX)
+	, InstructionIndex(INDEX_NONE)
 	{
 	}
 
-	FMultiplexJumpOp(uint64 InByteCodeIndex)
+	FMultiplexJumpOp(int32 InInstructionIndex)
 	: FMultiplexBaseOp(EMultiplexOpCode::Jump)
-	, ByteCodeIndex(InByteCodeIndex)
+	, InstructionIndex(InInstructionIndex)
 	{
 	}
 
-	uint64 ByteCodeIndex;
+	int32 InstructionIndex;
 };
 
 struct ANIMATIONCORE_API FMultiplexJumpIfTrueOp : public FMultiplexBaseOp
 {
 	FMultiplexJumpIfTrueOp()
 	: FMultiplexBaseOp(EMultiplexOpCode::JumpIfTrue)
-	, ByteCodeIndex(UINT64_MAX)
+	, InstructionIndex(INDEX_NONE)
 	, Condition()
 	{
 	}
 
-	FMultiplexJumpIfTrueOp(uint64 InByteCodeIndex, FMultiplexArgument InCondition)
+	FMultiplexJumpIfTrueOp(int32 InInstructionIndex, FMultiplexArgument InCondition)
 	: FMultiplexBaseOp(EMultiplexOpCode::JumpIfTrue)
-	, ByteCodeIndex(InByteCodeIndex)
+	, InstructionIndex(InInstructionIndex)
 	, Condition(InCondition)
 	{
 	}
 
-	uint64 ByteCodeIndex;
+	int32 InstructionIndex;
 	FMultiplexArgument Condition;
 };
 
@@ -198,19 +198,19 @@ struct ANIMATIONCORE_API FMultiplexJumpIfFalseOp : public FMultiplexBaseOp
 {
 	FMultiplexJumpIfFalseOp()
 	: FMultiplexBaseOp(EMultiplexOpCode::JumpIfFalse)
-	, ByteCodeIndex(UINT64_MAX)
+	, InstructionIndex(INDEX_NONE)
 	, Condition()
 	{
 	}
 
-	FMultiplexJumpIfFalseOp(uint64 InByteCodeIndex, FMultiplexArgument InCondition)
+	FMultiplexJumpIfFalseOp(int32 InInstructionIndex, FMultiplexArgument InCondition)
 	: FMultiplexBaseOp(EMultiplexOpCode::JumpIfFalse)
-	, ByteCodeIndex(InByteCodeIndex)
+	, InstructionIndex(InInstructionIndex)
 	, Condition(InCondition)
 	{
 	}
 
-	uint64 ByteCodeIndex;
+	int32 InstructionIndex;
 	FMultiplexArgument Condition;
 };
 
@@ -269,6 +269,8 @@ public:
 
 	FMultiplexByteCodeTable();
 
+	void Reset();
+	FORCEINLINE bool IsValidIndex(int32 InIndex) const { return Entries.IsValidIndex(InIndex); }
 	FORCEINLINE int32 Num() const { return Entries.Num(); }
 	FORCEINLINE const FMultiplexByteCodeTableEntry& operator[](int32 InIndex) const { return Entries[InIndex]; }
 
diff --git a/Engine/Source/Runtime/AnimationCore/Public/MultiplexRegistry.h b/Engine/Source/Runtime/AnimationCore/Public/MultiplexRegistry.h
--- a/Engine/Source/Runtime/AnimationCore/Public/MultiplexRegistry.h
+++ b/Engine/Source/Runtime/AnimationCore/Public/MultiplexRegistry.h
@@ -5,31 +5,7 @@
 #include "CoreMinimal.h"
 #include "MultiplexStorage.h"
 
-struct ANIMATIONCORE_API FMultiplexArgument
-{
-public:
-
-	FMultiplexArgument()
-		:Address(INDEX_NONE)
-	{
-	}
-
-	FMultiplexArgument(int32 InAddress)
-		:Address(InAddress)
-	{
-	}
-
-	FORCEINLINE bool operator ==(int32 InOther) const { return Address == InOther; }
-
-	FORCEINLINE bool IsLiteral() const { return Address < 0; }
-	FORCEINLINE int32 StorageType() const { return Address < 0 ? 1 : 0; }
-	FORCEINLINE int32 Index() const { return Address < 0 ? -(Address + 1) : Address; }
-
-private:
-	int32 Address;
-};
-
-typedef void (*FMultiplexFunctionPtr)(const TArrayView<FMultiplexArgument>&, FMultiplexStorage*, const TArrayView<void*>&);
+typedef void (*FMultiplexFunctionPtr)(const TArrayView<FMultiplexArgument>&, FMultiplexStorage**, const TArrayView<void*>&);
 
 struct ANIMATIONCORE_API FMultiplexFunction
 {
diff --git a/Engine/Source/Runtime/AnimationCore/Public/MultiplexStorage.h b/Engine/Source/Runtime/AnimationCore/Public/MultiplexStorage.h
--- a/Engine/Source/Runtime/AnimationCore/Public/MultiplexStorage.h
+++ b/Engine/Source/Runtime/AnimationCore/Public/MultiplexStorage.h
@@ -6,6 +6,41 @@
 #include "UObject/ObjectMacros.h"
 #include "MultiplexStorage.generated.h"
 
+struct ANIMATIONCORE_API FMultiplexArgument
+{
+public:
+
+	FMultiplexArgument()
+		:Address(INDEX_NONE)
+	{
+	}
+
+	FMultiplexArgument(int32 InAddress)
+		:Address(InAddress)
+	{
+	}
+
+	static FMultiplexArgument MakeArgument(int32 InAddress)
+	{
+		return FMultiplexArgument(InAddress);
+	}
+
+	static FMultiplexArgument MakeLiteral(int32 InAddress)
+	{
+		return FMultiplexArgument(-1 - InAddress);
+	}
+
+	FORCEINLINE bool operator ==(int32 InOther) const { return Address == InOther; }
+
+	FORCEINLINE bool IsLiteral() const { return Address < 0; }
+	FORCEINLINE int32 StorageType() const { return Address < 0 ? 1 : 0; }
+	FORCEINLINE int32 Index() const { return Address < 0 ? -(Address + 1) : Address; }
+	FORCEINLINE int32 NativeAddress() const { return Address; }
+
+private:
+	int32 Address;
+};
+
 UENUM()
 enum class EMultiplexAddressType : uint8
 {
@@ -112,6 +147,8 @@ public:
 	FORCEINLINE int32 Num() const { return Addresses.Num(); }
 	FORCEINLINE const FMultiplexAddress& operator[](int32 InIndex) const { return Addresses[InIndex]; }
 	FORCEINLINE FMultiplexAddress& operator[](int32 InIndex) { return Addresses[InIndex]; }
+	FORCEINLINE const FMultiplexAddress& operator[](const FMultiplexArgument& InArg) const { return Addresses[InArg.Index()]; }
+	FORCEINLINE FMultiplexAddress& operator[](const FMultiplexArgument& InArg) { return Addresses[InArg.Index()]; }
 	FORCEINLINE const FMultiplexAddress& operator[](const FName& InName) const { return Addresses[GetIndex(InName)]; }
 	FORCEINLINE FMultiplexAddress& operator[](const FName& InName) { return Addresses[GetIndex(InName)]; }
 
@@ -145,12 +182,24 @@ public:
 		return (const T*)&Data[Address.FirstByte()];
 	}
 
+	template<class T>
+	FORCEINLINE const T* Get(const FMultiplexArgument& InArgument) const
+	{
+		return Get<T>(InArgument.Index());
+	}
+
 	template<class T>
 	FORCEINLINE const T& GetRef(int32 InAddressIndex) const
 	{
 		return *Get<T>(InAddressIndex);
 	}
 
+	template<class T>
+	FORCEINLINE const T& GetRef(const FMultiplexArgument& InArgument) const
+	{
+		return GetRef<T>(InArgument.Index());
+	}
+
 	template<class T>
 	FORCEINLINE T* Get(int32 InAddressIndex)
 	{
@@ -160,12 +209,24 @@ public:
 		return (T*)&Data[Address.FirstByte()];
 	}
 
+	template<class T>
+	FORCEINLINE T* Get(const FMultiplexArgument& InArgument)
+	{
+		return Get<T>(InArgument.Index())
+	}
+
 	template<class T>
 	FORCEINLINE T& GetRef(int32 InAddressIndex)
 	{
 		return *Get<T>(InAddressIndex);
 	}
 
+	template<class T>
+	FORCEINLINE T& GetRef(const FMultiplexArgument& InArgument)
+	{
+		return GetRef<T>(InArgument.Index());
+	}
+
 	template<class T>
 	FORCEINLINE TArrayView<T> GetArray(int32 InAddressIndex)
 	{
@@ -174,6 +235,12 @@ public:
 		ensure(Address.ElementCount > 0);
 		return TArrayView<T>((T*)&Data[Address.FirstByte()], Address.ElementCount);
 	}
+	
+	template<class T>
+	FORCEINLINE TArrayView<T> GetArray(const FMultiplexArgument& InArgument)
+	{
+		return GetArray<T>(InArgument.Index());
+	}
 
 	FORCEINLINE UScriptStruct* GetScriptStruct(int32 InAddressIndex) const
 	{
diff --git a/Engine/Source/Runtime/AnimationCore/Public/MultiplexVM.h b/Engine/Source/Runtime/AnimationCore/Public/MultiplexVM.h
new file mode 100644
--- /dev/null
+++ b/Engine/Source/Runtime/AnimationCore/Public/MultiplexVM.h
@@ -0,0 +1,46 @@
+// Copyright 1998-2019 Epic Games, Inc. All Rights Reserved.
+
+#pragma once
+
+#include "CoreMinimal.h"
+#include "MultiplexStorage.h"
+#include "MultiplexRegistry.h"
+#include "MultiplexByteCode.h"
+#include "MultiplexVM.generated.h"
+
+UCLASS()
+class ANIMATIONCORE_API UMultiplexVM : public UObject
+{
+	GENERATED_BODY()
+
+public:
+
+	void Reset();
+
+	bool Execute(FMultiplexStorage** Storage = nullptr, TArrayView<void*> AdditionalArgs = TArrayView<void*>());
+
+	UFUNCTION()
+	int32 AddMultiplexFunction(UScriptStruct* InMultiplexStruct, const FName& InFunctionName);
+
+	UPROPERTY()
+	FMultiplexStorage Literals;
+
+	UPROPERTY()
+	FMultiplexStorage WorkState;
+
+	UPROPERTY()
+	FMultiplexByteCode ByteCode;
+
+	UPROPERTY()
+	FMultiplexByteCodeTable Instructions;
+
+private:
+
+	void ResolveFunctionsIfRequired();
+	void RefreshInstructionsIfRequired();
+
+	UPROPERTY()
+	TArray<FString> FunctionNames;
+
+	TArray<FMultiplexFunctionPtr> Functions;
+};
