Fix build break on OSX.

diff --git a/Engine/Source/Runtime/AnimationCore/Public/MultiplexStorage.h b/Engine/Source/Runtime/AnimationCore/Public/MultiplexStorage.h
--- a/Engine/Source/Runtime/AnimationCore/Public/MultiplexStorage.h
+++ b/Engine/Source/Runtime/AnimationCore/Public/MultiplexStorage.h
@@ -212,7 +212,7 @@ public:
 	template<class T>
 	FORCEINLINE T* Get(const FMultiplexArgument& InArgument)
 	{
-		return Get<T>(InArgument.Index())
+		return Get<T>(InArgument.Index());
 	}
 
 	template<class T>
