Remove unused function and nix dependency on Engine.

#rnx

diff --git a/Engine/Source/Runtime/AnimationCore/AnimationCore.Build.cs b/Engine/Source/Runtime/AnimationCore/AnimationCore.Build.cs
--- a/Engine/Source/Runtime/AnimationCore/AnimationCore.Build.cs
+++ b/Engine/Source/Runtime/AnimationCore/AnimationCore.Build.cs
@@ -6,11 +6,17 @@ public class AnimationCore : ModuleRules
 {
     public AnimationCore(ReadOnlyTargetRules Target) : base(Target)
     {
-        PublicDependencyModuleNames.AddRange(
+		PrivateIncludePathModuleNames.AddRange(
+			new string[]
+			{
+				"Engine"
+			}
+		);
+
+		PublicDependencyModuleNames.AddRange(
             new string[] {
                 "Core",
                 "CoreUObject",
-                "Engine",
             }
         );
     }
diff --git a/Engine/Source/Runtime/AnimationCore/Private/FABRIK.cpp b/Engine/Source/Runtime/AnimationCore/Private/FABRIK.cpp
--- a/Engine/Source/Runtime/AnimationCore/Private/FABRIK.cpp
+++ b/Engine/Source/Runtime/AnimationCore/Private/FABRIK.cpp
@@ -1,13 +1,6 @@
 // Copyright Epic Games, Inc. All Rights Reserved.
 
 #include "FABRIK.h"
-#include "BoneContainer.h"
-
-FVector FFABRIKChainLink::GetDirectionToParent(const FBoneContainer& BoneContainer, FCompactPoseBoneIndex BoneIndex)
-{
-	const FTransform& RefTransform = BoneContainer.GetRefPoseTransform(BoneIndex);
-	return -RefTransform.GetTranslation().GetSafeNormal();
-}
 
 namespace AnimationCore
 {
diff --git a/Engine/Source/Runtime/AnimationCore/Public/FABRIK.h b/Engine/Source/Runtime/AnimationCore/Public/FABRIK.h
--- a/Engine/Source/Runtime/AnimationCore/Public/FABRIK.h
+++ b/Engine/Source/Runtime/AnimationCore/Public/FABRIK.h
@@ -75,8 +75,6 @@ public:
 		, DefaultDirToParent(FVector(-1.f, 0.f, 0.f))
 	{
 	}
-
-	static ANIMATIONCORE_API FVector GetDirectionToParent(const FBoneContainer& BoneContainer, FCompactPoseBoneIndex BoneIndex);
 };
 
 namespace AnimationCore
diff --git a/Engine/Source/Runtime/Engine/Engine.Build.cs b/Engine/Source/Runtime/Engine/Engine.Build.cs
--- a/Engine/Source/Runtime/Engine/Engine.Build.cs
+++ b/Engine/Source/Runtime/Engine/Engine.Build.cs
@@ -98,6 +98,7 @@ public class Engine : ModuleRules
 
 		PrivateDependencyModuleNames.AddRange(
 			new string[] {
+				"AnimationCore",
 				"AppFramework",
 				"Networking",
 				"Landscape",
