New container to store per-vertex bone weights. Takes care of normalization, blending, capping weight counts, etc. while being lightweight and fast.

diff --git a/Engine/Source/Runtime/AnimationCore/Private/BoneWeights.cpp b/Engine/Source/Runtime/AnimationCore/Private/BoneWeights.cpp
new file mode 100644
--- /dev/null
+++ b/Engine/Source/Runtime/AnimationCore/Private/BoneWeights.cpp
@@ -0,0 +1,383 @@
+// Copyright Epic Games, Inc. All Rights Reserved.
+
+#include "BoneWeights.h"
+
+#include "Algo/IsSorted.h"
+#include "Misc/MemStack.h"
+
+
+static auto WeightSortPredicate = [](const FBoneWeight& A, const FBoneWeight& B) {
+	return A.GetRawWeight() > B.GetRawWeight();
+};
+
+
+bool FBoneWeights::SetBoneWeight(
+    FBoneWeight InBoneWeight,
+    const FBoneWeightsSettings& InSettings /*= {} */
+)
+{
+	// Does this bone already exist?
+	int32 WeightIndex = FindWeightIndexByBone(InBoneWeight.GetBoneIndex());
+
+	// If the sum of weights could possibly exceed 1.0, we may need normalization based on
+	// the weight settings.
+	bool bMayNeedNormalization = false;
+
+	if (WeightIndex != INDEX_NONE)
+	{
+		FBoneWeight& ExistingBoneWeight = BoneWeights[WeightIndex];
+		// New weight is below the threshold. Remove the current bone weight altogether.
+		if (InBoneWeight.GetRawWeight() < InSettings.GetRawWeightThreshold())
+		{
+			BoneWeights.RemoveAt(WeightIndex);
+
+			// If always normalizing, we need to re-normalize after removing this entry.
+			if (InSettings.GetNormalizeType() == EBoneWeightNormalizeType::Always)
+			{
+				NormalizeWeights(EBoneWeightNormalizeType::Always);
+			}
+
+			return false;
+		}
+
+		if (ExistingBoneWeight.GetRawWeight() == InBoneWeight.GetRawWeight())
+		{
+			return true;
+		}
+
+		bMayNeedNormalization = (ExistingBoneWeight.GetRawWeight() < InBoneWeight.GetRawWeight());
+
+		ExistingBoneWeight.SetRawWeight(InBoneWeight.GetRawWeight());
+	}
+	else
+	{
+		// If the new weight is below the threshold, reject and return.
+		if (InBoneWeight.GetRawWeight() < InSettings.GetRawWeightThreshold())
+		{
+			return false;
+		}
+
+		// Are we already at the limit of weights for this container?
+		if (BoneWeights.Num() == InSettings.GetMaxWeightCount())
+		{
+			// If the weight is smaller than the smallest weight currently, then we reject.
+			if (InBoneWeight.GetRawWeight() < BoneWeights.Last().GetRawWeight())
+			{
+				return false;
+			}
+
+			// Overwrite the last one, we'll put it in its correct place when we sort.
+			BoneWeights.Last() = InBoneWeight;
+		}
+		else
+		{
+			BoneWeights.Add(InBoneWeight);
+		}
+
+		bMayNeedNormalization = true;
+	}
+
+	// If we got here, then we updated/added weights. We're contractually obligated to keep the
+	// weights sorted.
+	SortWeights();
+
+	if ((InSettings.GetNormalizeType() == EBoneWeightNormalizeType::Always) ||
+	    (InSettings.GetNormalizeType() == EBoneWeightNormalizeType::AboveOne && bMayNeedNormalization))
+	{
+		Renormalize(InSettings);
+	}
+
+	return true;
+}
+
+
+bool FBoneWeights::RemoveBoneWeight(
+    FBoneIndexType InBoneIndex,
+    const FBoneWeightsSettings& InSettings /*= {} */
+)
+{
+	int32 WeightIndex = FindWeightIndexByBone(InBoneIndex);
+	if (WeightIndex == INDEX_NONE)
+	{
+		return false;
+	}
+
+	BoneWeights.RemoveAt(WeightIndex);
+
+	// Cull all weights that exceed limits set by the settings.
+	CullWeights(InSettings);
+
+	// Removing weights will always cause the weight sum to decrease, so we only have to normalize
+	// if always asked to.
+	if (InSettings.GetNormalizeType() == EBoneWeightNormalizeType::Always)
+	{
+		NormalizeWeights(EBoneWeightNormalizeType::Always);
+	}
+
+	return true;
+}
+
+
+void FBoneWeights::Renormalize(const FBoneWeightsSettings& InSettings /*= {}*/)
+{
+	NormalizeWeights(InSettings.GetNormalizeType());
+
+	// If entries are now below the threshold, remove them.
+	if (InSettings.GetNormalizeType() == EBoneWeightNormalizeType::Always && CullWeights(InSettings))
+	{
+		NormalizeWeights(EBoneWeightNormalizeType::Always);
+	}
+}
+
+
+FBoneWeights FBoneWeights::Create(
+    const FBoneIndexType InBones[MaxInlineBoneWeightCount],
+    const uint8 InInfluences[MaxInlineBoneWeightCount],
+    const FBoneWeightsSettings& InSettings /*= {} */
+)
+{
+	FBoneWeights Result;
+	Result.BoneWeights.Reserve(MaxInlineBoneWeightCount);
+
+	for (int32 Index = 0; Index < MaxInlineBoneWeightCount && InInfluences[Index]; Index++)
+	{
+		FBoneWeight BW(InBones[Index], InInfluences[Index]);
+		if (BW.GetRawWeight() >= InSettings.GetRawWeightThreshold())
+		{
+			Result.BoneWeights.Add(BW);
+		}
+	}
+
+	// Sort the weights by descending weight value before we clip it.
+	Result.SortWeights();
+
+	if (Result.BoneWeights.Num() > InSettings.GetMaxWeightCount())
+	{
+		Result.BoneWeights.SetNum(InSettings.GetMaxWeightCount());
+	}
+
+	Result.NormalizeWeights(InSettings.GetNormalizeType());
+
+	return Result;
+}
+
+
+FBoneWeights FBoneWeights::Create(
+    const FBoneIndexType* InBones,
+    const float* InInfluences,
+    int32 NumEntries,
+    const FBoneWeightsSettings& InSettings /*= {} */
+)
+{
+	TArray<FBoneWeight, BoneWeightsTempAllocatorT> StackBoneWeights;
+
+	StackBoneWeights.Reserve(NumEntries);
+	for (int32 Index = 0; Index < NumEntries; Index++)
+	{
+		FBoneWeight BW(InBones[Index], InInfluences[Index]);
+		if (BW.GetRawWeight() >= InSettings.GetRawWeightThreshold())
+		{
+			StackBoneWeights.Add(BW);
+		}
+	}
+
+	return CreateFromArrayView(StackBoneWeights, InSettings);
+}
+
+
+FBoneWeights FBoneWeights::Create(TArrayView<const FBoneWeight> BoneWeights, const FBoneWeightsSettings& InSettings /*= {} */)
+{
+	TArray<FBoneWeight, BoneWeightsTempAllocatorT> StackBoneWeights;
+
+	StackBoneWeights.Reserve(BoneWeights.Num());
+	for (const FBoneWeight& BW : BoneWeights)
+	{
+		if (BW.GetRawWeight() >= InSettings.GetRawWeightThreshold())
+		{
+			StackBoneWeights.Add(BW);
+		}
+	}
+
+	return CreateFromArrayView(StackBoneWeights, InSettings);
+}
+
+
+FBoneWeights FBoneWeights::CreateFromArrayView(
+    TArrayView<FBoneWeight> BoneWeights,
+    const FBoneWeightsSettings& InSettings /*= {} */
+)
+{
+	BoneWeights.Sort(WeightSortPredicate);
+
+	int32 NumEntries = FMath::Min(BoneWeights.Num(), InSettings.GetMaxWeightCount());
+	FBoneWeights Result;
+	if (NumEntries)
+	{
+		Result.BoneWeights = BoneWeights.Slice(0, NumEntries);
+	}
+	Result.NormalizeWeights(InSettings.GetNormalizeType());
+
+	return Result;
+}
+
+
+FBoneWeights FBoneWeights::Blend(
+    const FBoneWeights& InA,
+    const FBoneWeights& InB,
+    float InBias,
+    const FBoneWeightsSettings& InSettings /*= {} */
+)
+{
+	checkSlow(Algo::IsSorted(InA.BoneWeights, WeightSortPredicate));
+	checkSlow(Algo::IsSorted(InB.BoneWeights, WeightSortPredicate));
+
+	// Both empty?
+	if (InA.BoneWeights.Num() == 0 && InA.BoneWeights.Num() == 0)
+	{
+		return {};
+	}
+	// FIXME: We can probably special-case a few more fast paths. But let's collect statistics
+	// first.
+
+	// To simplify lookup and iteration over the two bone weight arrays, we sort by bone index
+	// value, but indirectly, since we can't sort them directly, as that would violate the
+	// sorted-by-descending-weight contract. Instead we create an indirection array on the stack
+	// and use that to iterate
+	auto CreateIndirectIndex = [](TArrayView<const FBoneWeight> BoneWeights, TArrayView<int32> IndexIndirect) {
+		for (int32 Index = 0; Index < IndexIndirect.Num(); Index++)
+		{
+			IndexIndirect[Index] = Index;
+		}
+		IndexIndirect.Sort([BoneWeights](int32 A, int32 B) {
+			return BoneWeights[A].GetBoneIndex() < BoneWeights[B].GetBoneIndex();
+		});
+	};
+
+	TArray<int32, BoneWeightsTempAllocatorT> IndirectIndexA;
+	IndirectIndexA.SetNumUninitialized(InA.BoneWeights.Num());
+	CreateIndirectIndex(InA.BoneWeights, IndirectIndexA);
+
+	TArray<int32, BoneWeightsTempAllocatorT> IndirectIndexB;
+	IndirectIndexB.SetNumUninitialized(InB.BoneWeights.Num());
+	CreateIndirectIndex(InB.BoneWeights, IndirectIndexB);
+
+	TArray<FBoneWeight, TInlineAllocator<MaxInlineBoneWeightCount * 2>> BoneWeights;
+	BoneWeights.Reserve(InA.BoneWeights.Num() + InB.BoneWeights.Num());
+
+	int32 RawBiasB = int32(InBias * float(FBoneWeight::MaxRawWeight));
+	int32 RawBiasA = FBoneWeight::MaxRawWeight - RawBiasB;
+
+	int32 IndexA = 0, IndexB = 0;
+	for (; IndexA < InA.BoneWeights.Num() && IndexB < InB.BoneWeights.Num(); /**/)
+	{
+		const FBoneWeight& BWA = InA.BoneWeights[IndirectIndexA[IndexA]];
+		const FBoneWeight& BWB = InB.BoneWeights[IndirectIndexB[IndexB]];
+
+		// If both have the same bone index, we blend them using the bias given and advance
+		// both arrays. If the bone indices differ, we copy from the array with the lower bone
+		// index value, to ensure we can possibly catch up with the other array. We then
+		// advance until we hit the end of either array after which we blindly copy the remains.
+		if (BWA.GetBoneIndex() == BWB.GetBoneIndex())
+		{
+			uint16 RawWeight = (BWA.GetRawWeight() * RawBiasA + BWB.GetRawWeight() * RawBiasB) / FBoneWeight::MaxRawWeight;
+
+			BoneWeights.Emplace(BWA.GetBoneIndex(), RawWeight);
+			IndexA++;
+			IndexB++;
+		}
+		else if (BWA.GetBoneIndex() < BWB.GetBoneIndex())
+		{
+			BoneWeights.Add(BWA);
+			IndexA++;
+		}
+		else
+		{
+			BoneWeights.Add(BWB);
+			IndexB++;
+		}
+	}
+
+	for (; IndexA < InA.BoneWeights.Num(); IndexA++)
+	{
+		BoneWeights.Add(InA.BoneWeights[IndirectIndexA[IndexA]]);
+	}
+	for (; IndexB < InB.BoneWeights.Num(); IndexB++)
+	{
+		BoneWeights.Add(InB.BoneWeights[IndirectIndexB[IndexB]]);
+	}
+
+	return CreateFromArrayView(BoneWeights, InSettings);
+}
+
+
+void FBoneWeights::SortWeights()
+{
+	BoneWeights.Sort(WeightSortPredicate);
+}
+
+
+bool FBoneWeights::CullWeights(const FBoneWeightsSettings& InSettings)
+{
+	bool bCulled = false;
+	if (BoneWeights.Num() > InSettings.GetMaxWeightCount())
+	{
+		BoneWeights.SetNum(InSettings.GetMaxWeightCount(), false);
+		bCulled = true;
+	}
+
+	// If entries are now below the threshold, remove them.
+	while (BoneWeights.Num() > 0 && BoneWeights.Last().GetRawWeight() < InSettings.GetRawWeightThreshold())
+	{
+		BoneWeights.SetNum(BoneWeights.Num() - 1, false);
+		bCulled = true;
+	}
+
+	return bCulled;
+}
+
+
+void FBoneWeights::NormalizeWeights(
+    EBoneWeightNormalizeType InNormalizeType)
+{
+	// Early checks
+	if (InNormalizeType == EBoneWeightNormalizeType::None || BoneWeights.Num() == 0)
+	{
+		return;
+	}
+
+	// Common case.
+	if (BoneWeights.Num() == 1)
+	{
+		if (InNormalizeType == EBoneWeightNormalizeType::Always)
+		{
+			BoneWeights[0].SetRawWeight(FBoneWeight::MaxRawWeight);
+		}
+		return;
+	}
+
+	// We operate on int64, since we can easily end up with wraparound issues during one of the
+	// multiplications below when using int32. This would tank the division by WeightSum.
+	int64 WeightSum = 0;
+	for (const FBoneWeight& BW : BoneWeights)
+	{
+		WeightSum += BW.GetRawWeight();
+	}
+
+	if ((InNormalizeType == EBoneWeightNormalizeType::Always && ensure(WeightSum != 0)) ||
+	    WeightSum > FBoneWeight::MaxRawWeight)
+	{
+		int64 Correction = 0;
+
+		// Here we treat the raw weight as a 16.16 fixed point value and ensure that the
+		// fraction, which would otherwise be lost through rounding, is carried over to the 
+		// subsequent values to maintain a constant sum to the max weight value.
+		// We do this in descending weight order in an attempt to ensure that weight values 
+		// aren't needlessly lost after scaling.
+		for (FBoneWeight& BW : BoneWeights)
+		{
+			int64 ScaledWeight = int64(BW.GetRawWeight()) * FBoneWeight::MaxRawWeight + Correction;
+			BW.SetRawWeight(uint16(FMath::Min(ScaledWeight / WeightSum, int64(FBoneWeight::MaxRawWeight))));
+			Correction = ScaledWeight - BW.GetRawWeight() * WeightSum;
+		}
+	}
+}
diff --git a/Engine/Source/Runtime/AnimationCore/Private/Tests/BoneWeightsTests.cpp b/Engine/Source/Runtime/AnimationCore/Private/Tests/BoneWeightsTests.cpp
new file mode 100644
--- /dev/null
+++ b/Engine/Source/Runtime/AnimationCore/Private/Tests/BoneWeightsTests.cpp
@@ -0,0 +1,510 @@
+// Copyright Epic Games, Inc. All Rights Reserved.
+
+#include "BoneWeights.h"
+#include "Misc/AutomationTest.h"
+#include "Serialization/BufferArchive.h"
+#include "Serialization/BufferReader.h"
+
+#include <initializer_list>
+
+#if WITH_DEV_AUTOMATION_TESTS
+
+IMPLEMENT_SIMPLE_AUTOMATION_TEST(FBoneWeightTestBasic, "System.AnimationCore.BoneWeight.Basic", EAutomationTestFlags::ApplicationContextMask | EAutomationTestFlags::SmokeFilter)
+bool FBoneWeightTestBasic::RunTest(const FString& Parameters)
+{
+	// Correct values (uint8)
+	{
+		FBoneWeight BWMin(FBoneIndexType(2), uint8(0));
+		UTEST_EQUAL(TEXT("FBoneWeight index (uint8)"), BWMin.GetBoneIndex(), 2);
+		UTEST_EQUAL(TEXT("FBoneWeight min value (uint8)"), BWMin.GetWeight(), 0.0f);
+		UTEST_EQUAL(TEXT("FBoneWeight min raw value (uint8)"), BWMin.GetRawWeight(), uint16(0));
+
+		// Because we construct a uin16 value from uint8 by doubling (v | v << 8). The reason
+		// is so that we end up with a fully saturated value (0xFFFF) for uint8(0xFF), rather 
+		// than (0xFF00) which would not be a full weight.
+		FBoneWeight BWMid(FBoneIndexType(0), uint8(127));
+		UTEST_EQUAL(TEXT("FBoneWeight mid value (uint8)"), BWMid.GetWeight(), 0x7F7F / 65535.0f);
+		UTEST_EQUAL(TEXT("FBoneWeight mid raw value (uint8)"), BWMid.GetRawWeight(), 0x7F7FU);
+
+		FBoneWeight BWMax(FBoneIndexType(0), uint8(255));
+		UTEST_EQUAL(TEXT("FBoneWeight max value (uint8)"), BWMax.GetWeight(), 1.0f);
+		UTEST_EQUAL(TEXT("FBoneWeight max raw value (uint8)"), BWMax.GetRawWeight(), FBoneWeight::MaxRawWeight);
+	}
+
+	// Correct values (uint16)
+	{
+		FBoneWeight BWMin(FBoneIndexType(2), uint16(0));
+		UTEST_EQUAL(TEXT("FBoneWeight index (uint8)"), BWMin.GetBoneIndex(), 2);
+		UTEST_EQUAL(TEXT("FBoneWeight min value (uint16)"), BWMin.GetWeight(), 0.0f);
+		UTEST_EQUAL(TEXT("FBoneWeight min raw value (uint16)"), BWMin.GetRawWeight(), uint16(0));
+
+		// Since uint16::max / 2 -is equal to 32767.5
+		FBoneWeight BWMid(FBoneIndexType(0), uint16(0x7FFF));
+		UTEST_EQUAL(TEXT("FBoneWeight mid value (uint16)"), BWMid.GetWeight(), 0x7FFF / float(FBoneWeight::MaxRawWeight));
+		UTEST_EQUAL(TEXT("FBoneWeight mid raw value (uint16)"), BWMid.GetRawWeight(), 0x7FFF);
+
+		FBoneWeight BWMax(FBoneIndexType(0), uint16(FBoneWeight::MaxRawWeight));
+		UTEST_EQUAL(TEXT("FBoneWeight max value (uint16)"), BWMax.GetWeight(), 1.0f);
+		UTEST_EQUAL(TEXT("FBoneWeight max raw value (uint16)"), BWMax.GetRawWeight(), FBoneWeight::MaxRawWeight);
+	}
+
+	// Copy constructor
+	{
+		FBoneWeight BWOrig(FBoneIndexType(2), uint16(0x8000));
+		FBoneWeight BWCopy(BWOrig);
+		UTEST_EQUAL(TEXT("FBoneWeight copied index"), BWCopy.GetBoneIndex(), 2);
+		UTEST_EQUAL(TEXT("FBoneWeight copied value"), BWCopy.GetRawWeight(), 0x8000);
+	}
+
+	// Archiving
+	{
+		FBufferArchive BA;
+		FBoneWeight BWOrig(FBoneIndexType(2), uint16(0x8000));
+		BA << BWOrig;
+
+		FBoneWeight BWRead;
+		FBufferReader BR(BA.GetData(), BA.Num(), /*bFreeOnClose=*/false);
+		BR << BWRead;
+		UTEST_EQUAL(TEXT("FBoneWeight::Archive[Index]"), BWRead.GetBoneIndex(), 2);
+		UTEST_EQUAL(TEXT("FBoneWeight::Archive[Weight]"), BWRead.GetRawWeight(), 0x8000);
+	}
+
+	return true;
+}
+
+IMPLEMENT_SIMPLE_AUTOMATION_TEST(FBoneWeightsSettingsTest, "System.AnimationCore.BoneWeights.Settings", EAutomationTestFlags::ApplicationContextMask | EAutomationTestFlags::SmokeFilter)
+bool FBoneWeightsSettingsTest::RunTest(const FString& Parameters)
+{
+	// Ensure that no-one's tampered with the defaults.
+	FBoneWeightsSettings S;
+
+	UTEST_EQUAL(TEXT("FBoneWeightsSettings NormalizeType"), S.GetNormalizeType(), EBoneWeightNormalizeType::Always);
+	UTEST_EQUAL(TEXT("FBoneWeightsSettings MaxWeightCount"), S.GetMaxWeightCount(), MaxInlineBoneWeightCount);
+	UTEST_EQUAL(TEXT("FBoneWeightsSettings Weight Threshold"), S.GetRawWeightThreshold(), uint16((1 << 8) | 1));
+
+	return true;
+}
+
+IMPLEMENT_SIMPLE_AUTOMATION_TEST(FBoneWeightsTestBasic, "System.AnimationCore.BoneWeights.Basic", EAutomationTestFlags::ApplicationContextMask | EAutomationTestFlags::SmokeFilter)
+bool FBoneWeightsTestBasic::RunTest(const FString& Parameters)
+{
+	{
+		const TCHAR* Name = TEXT("FBoneWeights::FBoneWeights");
+
+		FBoneWeights	A;
+		UTEST_EQUAL(Name, A.Num(), 0);
+
+		FBoneWeights	B(A);
+		UTEST_EQUAL(Name, B.Num(), 0);
+	}
+
+	{
+		const TCHAR* Name = TEXT("FBoneWeights::SetBoneWeight");
+
+		FBoneWeights	A;
+		A.SetBoneWeight(FBoneIndexType(0), 1.0f);
+		UTEST_EQUAL(Name, A.Num(), 1);
+
+		FBoneWeights	B(A);
+		UTEST_EQUAL(Name, B.Num(), 1);
+	}
+
+	{
+		const TCHAR* Name = TEXT("FBoneWeights::SetBoneWeight");
+
+		FBoneWeights	A;
+		A.SetBoneWeight(FBoneIndexType(0), 1.0f);
+
+		UTEST_EQUAL(Name, A.Num(), 1);
+
+		FBoneWeight BI = A[0];
+		UTEST_EQUAL(Name, BI.GetWeight(), 1.0f);
+		UTEST_EQUAL(Name, BI.GetRawWeight(), FBoneWeight::MaxRawWeight);
+	}
+
+	// Check normalization works
+	{
+		const TCHAR* Name = TEXT("FBoneWeights::SetBoneWeight[MaxWNormalize=AboveOne]");
+
+		FBoneWeightsSettings S;
+		S.SetNormalizeType(EBoneWeightNormalizeType::AboveOne);
+
+		FBoneWeights	A;
+		A.SetBoneWeight(FBoneIndexType(0), 1.0f, S);
+		A.SetBoneWeight(FBoneIndexType(1), 1.0f, S);
+
+		UTEST_EQUAL(Name, A.Num(), 2);
+		UTEST_EQUAL(Name, A[0].GetWeight(), 0.5f);
+	}
+	{
+		const TCHAR* Name = TEXT("FBoneWeights::SetBoneWeight[Normalize=Always]");
+
+		FBoneWeightsSettings S;
+		S.SetNormalizeType(EBoneWeightNormalizeType::Always);
+
+		FBoneWeights	A;
+		// Gets normalized to 1.0
+		A.SetBoneWeight(FBoneIndexType(0), 0.25f, S);
+
+		// Gets normalized to 0.25 / 1.25 = 0.2 leaving the first entry at 0.8
+		A.SetBoneWeight(FBoneIndexType(1), 0.25f, S);
+
+		UTEST_EQUAL(Name, A.Num(), 2);
+		UTEST_EQUAL(Name, A[0].GetWeight(), 0.8f);
+	}
+	{
+		const TCHAR* Name = TEXT("FBoneWeights::SetBoneWeight[Normalize=None]");
+
+		FBoneWeightsSettings S;
+		S.SetNormalizeType(EBoneWeightNormalizeType::None);
+
+		FBoneWeights	A;
+		A.SetBoneWeight(FBoneIndexType(0), 0.25f, S);
+		A.SetBoneWeight(FBoneIndexType(1), 0.25f, S);
+
+		UTEST_EQUAL(Name, A.Num(), 2);
+		UTEST_EQUAL(Name, A[0].GetWeight(), 0.25f);
+	}
+
+	// Check overwrite same bone index works  (with renormalization)
+	{
+		const TCHAR* Name = TEXT("FBoneWeights::SetBoneWeight[Normalize=AboveOne, Overwrite]");
+
+		FBoneWeightsSettings S;
+		S.SetNormalizeType(EBoneWeightNormalizeType::AboveOne);
+
+		FBoneWeights	A;
+		A.SetBoneWeight(FBoneIndexType(0), 0.25f, S);
+		A.SetBoneWeight(FBoneIndexType(1), 0.25f, S);
+
+		UTEST_EQUAL(Name, A.Num(), 2);
+		UTEST_EQUAL(Name, A[0].GetWeight(), 0.25f);
+
+		// This should set the total weight to above 1.0 and hence trigger normalization.
+		A.SetBoneWeight(FBoneIndexType(1), 1.0f, S);
+		UTEST_EQUAL(Name, A[0].GetWeight(), 0.8f);
+	}
+
+	// Check normalization correctly distributes error.
+	{
+		const TCHAR* Name = TEXT("FBoneWeights::SetBoneWeight[Normalize=Always, Distribute]");
+
+		// Ensure that no weights get thrown out because they're too small.
+		FBoneWeightsSettings SN, SA;
+		SN.SetNormalizeType(EBoneWeightNormalizeType::None);
+		SN.SetWeightThreshold(0.0f);
+		SA.SetNormalizeType(EBoneWeightNormalizeType::Always);
+		SA.SetWeightThreshold(0.0f);
+
+		FBoneWeights A;
+		for (uint32 Index = 0; Index < MaxInlineBoneWeightCount; Index++)
+		{
+			A.SetBoneWeight(FBoneIndexType(Index), Index == 0 ? 1.0f : 2.0f / float(FBoneWeight::MaxRawWeight), SN);
+		}
+
+		A.Renormalize(SA);
+
+		UTEST_EQUAL(Name, A.Num(), MaxInlineBoneWeightCount);
+		
+		// None of the small weights should not have changed in value, since the total scaling 
+		// (1.0 - 22/65535) is so close to 1.0 that it won't make a difference with 
+		// round-to-nearest.
+		int32 WeightSum = 0;
+		for (int32 Index = 0; Index < A.Num(); Index++)
+		{
+			WeightSum += int32(A[Index].GetRawWeight());
+			if (Index > 0)
+			{
+				FString NameSub = FString::Printf(TEXT("%s[%d]"), Name, Index);
+				UTEST_EQUAL(NameSub, int32(A[Index].GetRawWeight()), 2);
+			}
+		}
+		UTEST_EQUAL(Name, WeightSum, int32(FBoneWeight::MaxRawWeight));
+
+
+	}
+
+	// Check ordering works
+	{
+		const TCHAR* Name = TEXT("FBoneWeights::SetBoneWeight[Ordering]");
+
+		FBoneWeightsSettings S;
+		S.SetNormalizeType(EBoneWeightNormalizeType::None);
+
+		FBoneWeights	A;
+		A.SetBoneWeight(FBoneIndexType(0), 0.25f, S);
+		A.SetBoneWeight(FBoneIndexType(1), 0.5f, S);
+
+		UTEST_EQUAL(Name, A.Num(), 2);
+		UTEST_EQUAL(Name, A[0].GetBoneIndex(), 1);
+	}
+
+	// Check max limit works
+	{
+		const TCHAR* Name = TEXT("FBoneWeights::SetBoneWeight[MaxWeights]");
+
+		FBoneWeightsSettings S;
+		S.SetMaxWeightCount(3);
+		S.SetNormalizeType(EBoneWeightNormalizeType::None);
+
+		FBoneWeights	A;
+		A.SetBoneWeight(FBoneIndexType(0), 0.25f, S);
+		A.SetBoneWeight(FBoneIndexType(1), 0.20f, S);
+		A.SetBoneWeight(FBoneIndexType(2), 0.15f, S);
+		A.SetBoneWeight(FBoneIndexType(3), 0.10f, S);
+		A.SetBoneWeight(FBoneIndexType(4), 0.05f, S);
+
+		UTEST_EQUAL(Name, A.Num(), 3);
+		UTEST_EQUAL(Name, A[2].GetBoneIndex(), 2);
+
+		// Ensure it gets inserted just before the 2/0.15f value.
+		A.SetBoneWeight(FBoneIndexType(5), 0.17f, S);
+		UTEST_EQUAL(Name, A.Num(), 3);
+		UTEST_EQUAL(Name, A[2].GetBoneIndex(), 5);
+	}
+
+	// Check threshold works
+	{
+		const TCHAR* Name = TEXT("FBoneWeights::SetBoneWeight[Threshold, Normalize=None]");
+		FBoneWeightsSettings S;
+		S.SetWeightThreshold(0.40f);
+		S.SetNormalizeType(EBoneWeightNormalizeType::None);
+
+		FBoneWeights	A;
+		A.SetBoneWeight(FBoneIndexType(0), 1.00f, S);
+		A.SetBoneWeight(FBoneIndexType(1), 0.75f, S);
+		A.SetBoneWeight(FBoneIndexType(2), 0.50f, S);
+
+		UTEST_EQUAL(Name, A.Num(), 3);
+		UTEST_EQUAL(Name, A[2].GetBoneIndex(), 2);
+
+		A.SetBoneWeight(FBoneIndexType(3), 0.25f, S);
+		UTEST_EQUAL(Name, A.Num(), 3);
+		UTEST_EQUAL(Name, A[2].GetBoneIndex(), 2);
+
+		A.SetBoneWeight(FBoneIndexType(4), 0.45f, S);
+		UTEST_EQUAL(Name, A.Num(), 4);
+		UTEST_EQUAL(Name, A[3].GetBoneIndex(), 4);
+
+		// Set an existing bone index to be below threshold. It should get removed.
+		A.SetBoneWeight(FBoneIndexType(1), 0.1f, S);
+		UTEST_EQUAL(Name, A.Num(), 3);
+		UTEST_EQUAL(Name, A[2].GetBoneIndex(), 4);
+	}
+
+	{
+		const TCHAR* Name = TEXT("FBoneWeights::SetBoneWeight[Threshold, Normalize=Always]");
+		FBoneWeightsSettings S;
+		S.SetWeightThreshold(0.25f);
+		S.SetNormalizeType(EBoneWeightNormalizeType::Always);
+
+		FBoneWeights	A;
+		A.SetBoneWeight(FBoneIndexType(0), 1.00f, S);
+		// Weights = [1.0]
+		UTEST_EQUAL(Name, A.Num(), 1);
+		UTEST_EQUAL(Name, A[0].GetBoneIndex(), 0);
+
+
+		// Add one more, this should pass the threshold, both before and after normalization.
+		A.SetBoneWeight(FBoneIndexType(1), 0.5f, S);
+		// Weights = [0.666, 0.333]
+		UTEST_EQUAL(Name, A.Num(), 2);
+		UTEST_EQUAL(Name, A[1].GetBoneIndex(), 1);
+
+
+		// Add one more, this passes the early threshold check, but not after normalization.
+		A.SetBoneWeight(FBoneIndexType(2), 0.25f, S);
+		// Weights = [0.5328, 0.2664, 0.2] -> Last entry culled.
+		UTEST_EQUAL(Name, A.Num(), 2);
+		UTEST_EQUAL(Name, A[1].GetBoneIndex(), 1);
+	}
+
+	// Check removal works
+	{
+		const TCHAR* Name = TEXT("FBoneWeights::RemoveBoneWeight");
+
+		// Don't normalize at first.
+		FBoneWeights A;
+		A.SetBoneWeight(FBoneIndexType(0), 0.5f);
+		A.SetBoneWeight(FBoneIndexType(1), 0.5f);
+
+		UTEST_EQUAL(Name, A.Num(), 2);
+
+		A.RemoveBoneWeight(FBoneIndexType(2));
+		UTEST_EQUAL(Name, A.Num(), 2);
+
+		// Normalize on remove to see if the weight changes.
+		FBoneWeightsSettings S;
+		S.SetNormalizeType(EBoneWeightNormalizeType::Always);
+		A.RemoveBoneWeight(FBoneIndexType(0), S);
+
+		UTEST_EQUAL(Name, A.Num(), 1);
+		UTEST_EQUAL(Name, A[0].GetWeight(), 1.0f);
+	}
+
+	return true;
+}
+
+IMPLEMENT_SIMPLE_AUTOMATION_TEST(FBoneWeightsTestCreate, "System.AnimationCore.BoneWeights.Create", EAutomationTestFlags::ApplicationContextMask | EAutomationTestFlags::SmokeFilter)
+bool FBoneWeightsTestCreate::RunTest(const FString& Parameters)
+{
+	{
+		static const FBoneIndexType Bones[MaxInlineBoneWeightCount] = { 0, 1, 2, 3, 0, 0, 0, 0, 0, 0, 0, 0 };
+		static const uint8 Influences[MaxInlineBoneWeightCount] = { 255, 31, 63, 127, 0, 0, 0, 0, 0, 0, 0, 0 };
+
+		{
+			const TCHAR* Name = TEXT("FBoneWeights::Create[FSoftSkinVertex, Threshold]");
+			FBoneWeightsSettings S;
+			S.SetWeightThreshold(0.20);
+			FBoneWeights BW = FBoneWeights::Create(Bones, Influences, S);
+
+			UTEST_EQUAL(Name, BW.Num(), 3);
+			UTEST_EQUAL(Name, BW[2].GetBoneIndex(), 2);
+		}
+		{
+			const TCHAR* Name = TEXT("FBoneWeights::Create[FSoftSkinVertex, MaxWeights]");
+			FBoneWeightsSettings S;
+			S.SetMaxWeightCount(2);
+			FBoneWeights BW = FBoneWeights::Create(Bones, Influences, S);
+
+			UTEST_EQUAL(Name, BW.Num(), 2);
+			UTEST_EQUAL(Name, BW[1].GetBoneIndex(), 3);
+		}
+	}
+
+	{
+		static const FBoneIndexType Bones[] = { 0, 1, 2, 3};
+		static const float Influences[] = { 1.0f, 0.125f, 0.25f, 0.5f};
+
+		static const int32 ArraySize = sizeof(Bones) / sizeof(Bones[0]);
+
+		{
+			const TCHAR* Name = TEXT("FBoneWeights::Create[Arrays, Threshold]");
+			FBoneWeightsSettings S;
+			S.SetWeightThreshold(0.20);
+			FBoneWeights BW = FBoneWeights::Create(Bones, Influences, ArraySize, S);
+
+			UTEST_EQUAL(Name, BW.Num(), 3);
+			UTEST_EQUAL(Name, BW[2].GetBoneIndex(), 2);
+		}
+		{
+			const TCHAR* Name = TEXT("FBoneWeights::Create[Arrays, MaxWeights]");
+			FBoneWeightsSettings S;
+			S.SetMaxWeightCount(2);
+			FBoneWeights BW = FBoneWeights::Create(Bones, Influences, ArraySize, S);
+
+			UTEST_EQUAL(Name, BW.Num(), 2);
+			UTEST_EQUAL(Name, BW[1].GetBoneIndex(), 3);
+		}
+	}
+	return true;
+}
+
+struct FSimpleWeight
+{
+	FSimpleWeight(int Index, float Weight) : Index(Index), Weight(Weight) {}
+	int Index;
+	float Weight;
+};
+
+static FBoneWeights CreateWeights(
+	std::initializer_list<FSimpleWeight> InWeights
+	)
+{
+	// We only add raw values.
+	FBoneWeightsSettings S;
+	S.SetNormalizeType(EBoneWeightNormalizeType::None);
+
+	TArray<FBoneWeight> BWS;
+	BWS.Reserve(InWeights.size());
+	for (const FSimpleWeight& SW : InWeights)
+	{
+		BWS.Add(FBoneWeight(FBoneIndexType(SW.Index), SW.Weight));
+	}
+
+	return FBoneWeights::Create(BWS, S);
+}
+
+
+IMPLEMENT_SIMPLE_AUTOMATION_TEST(FBoneWeightsTestBlend, "System.AnimationCore.BoneWeights.Blend", EAutomationTestFlags::ApplicationContextMask | EAutomationTestFlags::SmokeFilter)
+bool FBoneWeightsTestBlend::RunTest(const FString& Parameters)
+{
+	{
+		const TCHAR* Name = TEXT("FBoneWeights::Blend[One Each, Disjoint]");
+
+		auto BW1 = CreateWeights({ {0, 0.6f} });
+		auto BW2 = CreateWeights({ {1, 0.4f} });
+
+		auto BW = FBoneWeights::Blend(BW1, BW2, 0.5f);
+
+		UTEST_EQUAL(Name, BW.Num(), 2);
+		UTEST_EQUAL(Name, BW[0].GetBoneIndex(), 0);
+		UTEST_EQUAL(Name, BW[1].GetBoneIndex(), 1);
+	}
+
+	{
+		const TCHAR* Name = TEXT("FBoneWeights::Blend[One Each, Overlap]");
+
+		auto BW1 = CreateWeights({ {0, 0.9f} });
+		auto BW2 = CreateWeights({ {0, 0.1f} });
+
+		FBoneWeightsSettings SN;
+		SN.SetNormalizeType(EBoneWeightNormalizeType::None);
+		auto BWN = FBoneWeights::Blend(BW1, BW2, 0.5f, SN);
+
+		// With no normalization, the blended value will be 0.5
+		UTEST_EQUAL(Name, BWN.Num(), 1);
+		UTEST_EQUAL(Name, BWN[0].GetWeight(), 0.5f);
+
+		FBoneWeightsSettings SA;
+		SA.SetNormalizeType(EBoneWeightNormalizeType::Always);
+		auto BWA = FBoneWeights::Blend(BW1, BW2, 0.5f, SA);
+
+		// With full normalization, the blended value will be 1.0
+		UTEST_EQUAL(Name, BWA.Num(), 1);
+		UTEST_EQUAL(Name, BWA[0].GetWeight(), 1.0f);
+	}
+
+	{
+		const TCHAR* Name = TEXT("FBoneWeights::Blend[Two Each, One Overlap]");
+
+		auto BW1 = CreateWeights({ {0, 1.0f}, {1, 0.5f} });
+		auto BW2 = CreateWeights({ {1, 0.5f}, {2, 0.25f} });
+
+		FBoneWeightsSettings S;
+		S.SetNormalizeType(EBoneWeightNormalizeType::None);
+		auto BW = FBoneWeights::Blend(BW1, BW2, 0.5f, S);
+
+		UTEST_EQUAL(Name, BW.Num(), 3);
+		UTEST_EQUAL(Name, BW[0].GetBoneIndex(), 0);
+		UTEST_EQUAL(Name, BW[0].GetWeight(), 1.0f);
+		UTEST_EQUAL(Name, BW[1].GetBoneIndex(), 1);
+		UTEST_EQUAL(Name, BW[1].GetWeight(), 0.5f);
+		UTEST_EQUAL(Name, BW[2].GetBoneIndex(), 2);
+		UTEST_EQUAL(Name, BW[2].GetWeight(), 0.25f);
+	}
+
+	{
+		const TCHAR* Name = TEXT("FBoneWeights::Blend[Two/Three, Two Overlap]");
+
+		auto BW1 = CreateWeights({ {0, 1.0f}, {1, 0.5f} });
+		auto BW2 = CreateWeights({ {0, 0.5f}, {2, 0.5f}, {1, 0.25f} });
+
+		FBoneWeightsSettings S;
+		S.SetNormalizeType(EBoneWeightNormalizeType::None);
+		auto BW = FBoneWeights::Blend(BW1, BW2, 0.5f, S);
+
+		UTEST_EQUAL(Name, BW.Num(), 3);
+		UTEST_EQUAL(Name, BW[0].GetBoneIndex(), 0);
+		UTEST_EQUAL(Name, BW[0].GetWeight(), 0.75f);
+		UTEST_EQUAL(Name, BW[1].GetBoneIndex(), 2);
+		UTEST_EQUAL(Name, BW[1].GetWeight(), 0.5f);
+		UTEST_EQUAL(Name, BW[2].GetBoneIndex(), 1);
+		UTEST_EQUAL(Name, BW[2].GetWeight(), 0.375f);
+	}
+
+	return true;
+}
+
+
+#endif //WITH_DEV_AUTOMATION_TESTS
diff --git a/Engine/Source/Runtime/AnimationCore/Public/BoneWeights.h b/Engine/Source/Runtime/AnimationCore/Public/BoneWeights.h
new file mode 100644
--- /dev/null
+++ b/Engine/Source/Runtime/AnimationCore/Public/BoneWeights.h
@@ -0,0 +1,399 @@
+// Copyright Epic Games, Inc. All Rights Reserved.
+
+#pragma once
+
+#include "BoneIndices.h"
+#include "Containers/UnrealString.h"
+#include "GPUSkinPublicDefs.h"
+#include "Math/NumericLimits.h"
+#include "Serialization/Archive.h"
+#include "Templates/TypeHash.h"
+
+/**
+The maximum number of inline bone weights.
+*/
+static constexpr int32 MaxInlineBoneWeightCount = MAX_TOTAL_INFLUENCES;
+
+class FBoneWeight
+{
+public:
+	/**
+	* Default constructor. Deliberately leaves the member values uninitialized for performance.
+	*/
+	FBoneWeight() = default;
+
+	FBoneWeight(const FBoneWeight&) = default;
+	FBoneWeight(FBoneWeight&&) noexcept = default;
+
+	FBoneWeight& operator=(const FBoneWeight&) = default;
+	FBoneWeight& operator=(FBoneWeight&&) = default;
+
+	/**
+	* Returns true if this object's bone weight values are equal to the other container's values.
+	* @note Only equality comparison are supported. Relational comparisons are meaningless.
+	*/
+	bool operator==(const FBoneWeight& InBoneWeight) const
+	{
+		return BoneIndex == InBoneWeight.BoneIndex && RawWeight == InBoneWeight.RawWeight;
+	}
+
+	/**
+	* Returns true if this object's bone weight values are not equal to the other container's values.
+	* @note Only equality comparison are supported. Relational comparisons are meaningless.
+	*/
+	bool operator!=(const FBoneWeight& InBoneWeight) const
+	{
+		return BoneIndex != InBoneWeight.BoneIndex || RawWeight != InBoneWeight.RawWeight;
+	}
+
+	/**
+	* The maximum raw weight value for a bone influence.
+	*/
+	static constexpr int32 MaxRawWeight = TNumericLimits<uint16>::Max();
+
+	/**
+	* A constructor for old-style bone weights where the weight was stored as an unsigned eight
+	* bit integer.
+	*/
+	explicit FBoneWeight(FBoneIndexType InBoneIndex, uint8 InWeight)
+	    : BoneIndex(InBoneIndex), RawWeight((InWeight << 8) | InWeight)
+	{
+	}
+
+	/**
+	* A constructor for bone weights where the weight is stored as an unsigned sixteen
+	* bit integer. This is the natural storage format for this container.
+	*/
+	explicit FBoneWeight(FBoneIndexType InBoneIndex, uint16 InRawWeight)
+	    : BoneIndex(InBoneIndex), RawWeight(InRawWeight)
+	{
+	}
+
+	/**
+	* A constructor for bone weights which converts the weight from a float value in the
+	* [0-1] range. Values outside of that range are clamped to the range.
+	*/
+	explicit FBoneWeight(FBoneIndexType InBoneIndex, float InWeight)
+	    : BoneIndex(InBoneIndex)
+	{
+		SetWeight(InWeight);
+	}
+
+	/**
+	* Set the bone stored index.
+	*/
+	void SetBoneIndex(FBoneIndexType InBoneIndex)
+	{
+		BoneIndex = InBoneIndex;
+	}
+
+	/**
+	* Returns the bone stored index.
+	*/
+	FBoneIndexType GetBoneIndex() const
+	{
+		return BoneIndex;
+	}
+
+	/**
+	* Set stored weight as a float. Values outside the [0-1] range are clamped to that range.
+	* Undefined float values will result in an undefined weight.
+	*/
+	void SetWeight(float InWeight)
+	{
+		InWeight = FMath::Clamp(InWeight, 0.0f, 1.0f);
+		RawWeight = uint16(InWeight * float(MaxRawWeight) + 0.5f);
+	}
+
+	/**
+	* Returns the stored weight value as a float in the [0-1] range.
+	*/
+	float GetWeight() const
+	{
+		return RawWeight / float(MaxRawWeight);
+	}
+
+	/**
+	* Set the stored weight in the raw format. This avoids any floating point conversion.
+	*/
+	void SetRawWeight(uint16 InRawWeight)
+	{
+		RawWeight = InRawWeight;
+	}
+
+	/**
+	* Returns the stored weight in the container's raw format, avoiding floating point
+	* conversion.
+	*/
+	uint16 GetRawWeight() const
+	{
+		return RawWeight;
+	}
+
+	friend FArchive& operator<<(FArchive& Ar, FBoneWeight& V)
+	{
+		Ar << V.BoneIndex;
+		Ar << V.RawWeight;
+
+		return Ar;
+	}
+
+	/**
+	* A helper function to return a human-readable version of the bone weight.
+	*/
+	FString ToString() const
+	{
+		return FString::Printf(TEXT("<%d, %g>"), BoneIndex, GetWeight());
+	}
+
+private:
+	friend uint32 GetTypeHash(const FBoneWeight& InBoneWeight);
+
+	FBoneIndexType BoneIndex;
+	uint16 RawWeight;
+};
+
+/**
+* A hashing function to allow the FBoneWeight class to be used with hashing containers (e.g.
+* TSet or TMap).
+*/
+inline uint32 GetTypeHash(const FBoneWeight& InBoneWeight)
+{
+	return HashCombine(GetTypeHash(uint16(InBoneWeight.BoneIndex)),
+	    GetTypeHash(InBoneWeight.RawWeight));
+}
+
+
+/** Specifies the method for which the bone weights are normalized after the bone weight
+  * list is modified.
+  */
+enum class EBoneWeightNormalizeType
+{
+	/** No normalization is performed. The sum of the weight values can exceed 1.0 */
+	None,
+
+	/** Normalization is only performed if the sum of the weights exceeds 1.0 */
+	AboveOne,
+
+	/** Normalization is always performed such that the sum of the weights is always equal to 1.0 */
+	Always
+};
+
+
+class FBoneWeightsSettings
+{
+public:
+	FBoneWeightsSettings() = default;
+
+	/** Set the normalization type when manipulating the weight values in FBoneWeights */
+	FBoneWeightsSettings& SetNormalizeType(EBoneWeightNormalizeType InNormalizeType)
+	{
+		NormalizeType = InNormalizeType;
+		return *this;
+	}
+
+	/** Returns the current normalization type for these settings */
+	EBoneWeightNormalizeType GetNormalizeType() const { return NormalizeType; }
+
+	/** Sets the maximum number of weights that can be applied to a single FBoneWeights object.
+	  * When weights are added, the smallest weights, past this limit, are dropped.
+	  */
+	FBoneWeightsSettings& SetMaxWeightCount(int32 InMaxWeightCount)
+	{
+		MaxWeightCount = FMath::Max(1, InMaxWeightCount);
+		return *this;
+	}
+
+	/** Returns the maximum number of weights for these settings */
+	int32 GetMaxWeightCount() const { return MaxWeightCount; }
+
+	/** Sets the minimum influence allowed. Any bone weight value below this limit will be
+	  * ignored. The threshold value is clamped to the half-closed interval (0, 1] since
+	  * weight values of zero indicate no influence at all and are ignored completely.
+	  */
+	FBoneWeightsSettings& SetWeightThreshold(float InWeightThreshold)
+	{
+		InWeightThreshold = FMath::Clamp(InWeightThreshold, 0.0f, 1.0f);
+		WeightThreshold = uint16(InWeightThreshold * FBoneWeight::MaxRawWeight + 0.5f);
+		WeightThreshold = FMath::Max(WeightThreshold, uint16(1));
+		return *this;
+	}
+
+	/** Returns the weight threshold as a float value between (0, 1]. */
+	float GetWeightThreshold() const
+	{
+		return WeightThreshold / float(TNumericLimits<uint16>::Max());
+	}
+
+	/** Returns the raw weight threshold. This is the value used internally for weight
+	  * computation.
+	  */
+	uint16 GetRawWeightThreshold() const
+	{
+		return WeightThreshold;
+	}
+
+private:
+	EBoneWeightNormalizeType NormalizeType = EBoneWeightNormalizeType::Always;
+	int32 MaxWeightCount = MaxInlineBoneWeightCount;
+	uint16 WeightThreshold = 257; // == uint8(1) converted to 16 bit using 'v | v << 8'.
+};
+
+/// A simple container for per-vertex influence of bones and their weights.
+class FBoneWeights
+{
+	using BoneWeightsTempAllocatorT = TInlineAllocator<MaxInlineBoneWeightCount>;
+	using BoneWeightArrayT = TArray<FBoneWeight, BoneWeightsTempAllocatorT>;
+
+public:
+	FBoneWeights() = default;
+
+	/**
+	  * Returns true if all of this container's values and count are equal to the other
+	  * container's values and count.
+	  * @note Only equality comparison are supported. Relational comparisons are meaningless.
+	  */
+	bool operator==(const FBoneWeights& InBoneWeight) const
+	{
+		return BoneWeights == InBoneWeight.BoneWeights;
+	}
+
+	/**
+	* Returns true if this container's values are equal to the other container's values.
+	  @note Only equality comparison are supported. Relational comparisons are meaningless.
+	*/
+	bool operator!=(const FBoneWeights& InBoneWeight) const
+	{
+		return BoneWeights != InBoneWeight.BoneWeights;
+	}
+
+	/** Set a new bone weight. If an existing weight exists with the same bone index, it's
+  	  * weight value is replaced with the weight value of the given entry. Otherwise a 
+      * new index is added. In both cases the new entry is subject to the given settings, which
+	  * may influence what the entry's final weight value is and whether it gets thrown away
+      * for not passing the threshold value. If the weight was successfully incorporated, then
+	  * this function returns true. Otherwise it returns false.
+	  */
+	ANIMATIONCORE_API bool SetBoneWeight(
+	    FBoneWeight InBoneWeight,
+	    const FBoneWeightsSettings& InSettings = {});
+
+	bool SetBoneWeight(
+	    FBoneIndexType InBone,
+	    float InWeight,
+	    const FBoneWeightsSettings& InSettings = {})
+	{
+		return SetBoneWeight(FBoneWeight{InBone, InWeight}, InSettings);
+	}
+
+	/** Removes a specific bone from the list of weights, re-normalizing and pruning bones,
+	  * if needed.
+	  */
+	ANIMATIONCORE_API bool RemoveBoneWeight(
+	    FBoneIndexType InBone,
+	    const FBoneWeightsSettings& InSettings = {});
+
+	/** Force normalization of weights. This is useful if a set of operations were performed
+	  * with no normalization, for efficiency, and normalization is needed post-operation.
+	  */
+	ANIMATIONCORE_API void Renormalize(const FBoneWeightsSettings& InSettings = {});
+
+	/** A helper to create a FBoneWeights container from FSoftSkinVertex data structure. */
+	ANIMATIONCORE_API static FBoneWeights Create(
+	    const FBoneIndexType InBones[MaxInlineBoneWeightCount],
+	    const uint8 InWeights[MaxInlineBoneWeightCount],
+	    const FBoneWeightsSettings& InSettings = {});
+
+	/** A helper to create a FBoneWeights container from separated bone index and weight arrays.
+	  * The size of the two arrays *must* be the same -- otherwise the behavior is undefined.
+	  */
+	ANIMATIONCORE_API static FBoneWeights Create(
+	    const FBoneIndexType* InBones,
+	    const float* InWeights,
+	    int32 NumEntries,
+	    const FBoneWeightsSettings& InSettings = {});
+
+	/** A helper to create FBoneWeights container from a TArray of FBoneWeight objects. */
+	ANIMATIONCORE_API static FBoneWeights Create(
+	    TArrayView<const FBoneWeight> BoneWeights,
+	    const FBoneWeightsSettings& InSettings = {});
+
+	/** Blend two bone weights together, making sure to add in every influence from both,
+	  * using the given settings. The bias value should lie on the [0,1] interval. Values outside
+	  * that range may give unwanted results. 
+	  */
+	ANIMATIONCORE_API static FBoneWeights Blend(
+	    const FBoneWeights& InA,
+	    const FBoneWeights& InB,
+	    float InBias,
+	    const FBoneWeightsSettings& InSettings = {});
+
+	// Ranged-based for loop compatibility -- but only the const version.
+	using RangedForConstIteratorType = BoneWeightArrayT::RangedForConstIteratorType;
+
+	RangedForConstIteratorType begin() const { return BoneWeights.begin(); }
+	RangedForConstIteratorType end() const { return BoneWeights.end(); }
+
+	/** Return the number of bone weights in this container.
+	  * @return The number of bone weights.
+	  */
+	int32 Num() const { return BoneWeights.Num(); }
+
+	/** Return the weight at index Index. Using an index value less than zero -- or equal
+	  * to or greater than the result of Num() -- is an undefined operation.
+	  * @param Index The index of the desired bone weight value.
+	  * @return The bone weight value for that index, or undefined if outside the valid range.
+	  */
+	const FBoneWeight& operator[](int32 Index) { return BoneWeights.operator[](Index); }
+
+
+	// -- Helper functions
+
+	/** Find the bone weight corresponding to the given bone index. If a bone weight with this
+	  * index doesn't exist
+	  */
+	int32 FindWeightIndexByBone(FBoneIndexType InBoneIndex) const
+	{
+		return BoneWeights.IndexOfByPredicate(
+		    [InBoneIndex](const FBoneWeight& BW) { return InBoneIndex == BW.GetBoneIndex(); });
+	}
+
+	friend FArchive& operator<<(FArchive& Ar, FBoneWeights& V)
+	{
+		Ar << V.BoneWeights;
+
+		return Ar;
+	}
+
+	FString ToString() const
+	{
+		FString Result(TEXT("["));
+		Result.Append(FString::JoinBy(BoneWeights, TEXT(", "), [](const FBoneWeight& V) { return V.ToString(); }));
+		Result.Append(TEXT("]"));
+		return Result;
+	}
+
+private:
+	static FBoneWeights CreateFromArrayView(
+	    TArrayView<FBoneWeight> BoneWeights,
+	    const FBoneWeightsSettings& InSettings = {});
+
+	void SortWeights();
+	bool CullWeights(const FBoneWeightsSettings& InSettings);
+	void NormalizeWeights(EBoneWeightNormalizeType InNormalizeType);
+
+	friend uint32 GetTypeHash(const FBoneWeights& InBoneWeights);
+
+	/// List of bone weights, in order of descending weight.
+	BoneWeightArrayT BoneWeights;
+};
+
+inline uint32 GetTypeHash(const FBoneWeights& InBoneWeights)
+{
+	uint32 Hash = GetTypeHash(InBoneWeights.BoneWeights.Num());
+	for (const FBoneWeight& BoneWeight : InBoneWeights)
+	{
+		Hash = HashCombine(Hash, GetTypeHash(BoneWeight));
+	}
+	return Hash;
+}
