Build fix for Linux -- lambda argument shadows member variable.

#rnx

diff --git a/Engine/Source/Runtime/AnimationCore/Private/BoneWeights.cpp b/Engine/Source/Runtime/AnimationCore/Private/BoneWeights.cpp
--- a/Engine/Source/Runtime/AnimationCore/Private/BoneWeights.cpp
+++ b/Engine/Source/Runtime/AnimationCore/Private/BoneWeights.cpp
@@ -243,13 +243,13 @@ FBoneWeights FBoneWeights::Blend(
 	// value, but indirectly, since we can't sort them directly, as that would violate the
 	// sorted-by-descending-weight contract. Instead we create an indirection array on the stack
 	// and use that to iterate
-	auto CreateIndirectIndex = [](TArrayView<const FBoneWeight> BoneWeights, TArrayView<int32> IndexIndirect) {
-		for (int32 Index = 0; Index < IndexIndirect.Num(); Index++)
+	auto CreateIndirectIndex = [](TArrayView<const FBoneWeight> InBoneWeights, TArrayView<int32> InIndexIndirect) {
+		for (int32 Index = 0; Index < InIndexIndirect.Num(); Index++)
 		{
-			IndexIndirect[Index] = Index;
+			InIndexIndirect[Index] = Index;
 		}
-		IndexIndirect.Sort([BoneWeights](int32 A, int32 B) {
-			return BoneWeights[A].GetBoneIndex() < BoneWeights[B].GetBoneIndex();
+		InIndexIndirect.Sort([InBoneWeights](int32 A, int32 B) {
+			return InBoneWeights[A].GetBoneIndex() < InBoneWeights[B].GetBoneIndex();
 		});
 	};
 
