Fix duplicate check to be separate checks. Added a test for it.

#rnx

diff --git a/Engine/Source/Runtime/AnimationCore/Private/BoneWeights.cpp b/Engine/Source/Runtime/AnimationCore/Private/BoneWeights.cpp
--- a/Engine/Source/Runtime/AnimationCore/Private/BoneWeights.cpp
+++ b/Engine/Source/Runtime/AnimationCore/Private/BoneWeights.cpp
@@ -232,7 +232,7 @@ FBoneWeights FBoneWeights::Blend(
 	checkSlow(Algo::IsSorted(InB.BoneWeights, WeightSortPredicate));
 
 	// Both empty?
-	if (InA.BoneWeights.Num() == 0 && InA.BoneWeights.Num() == 0)
+	if (InA.BoneWeights.Num() == 0 && InB.BoneWeights.Num() == 0)
 	{
 		return {};
 	}
diff --git a/Engine/Source/Runtime/AnimationCore/Private/Tests/BoneWeightsTests.cpp b/Engine/Source/Runtime/AnimationCore/Private/Tests/BoneWeightsTests.cpp
--- a/Engine/Source/Runtime/AnimationCore/Private/Tests/BoneWeightsTests.cpp
+++ b/Engine/Source/Runtime/AnimationCore/Private/Tests/BoneWeightsTests.cpp
@@ -429,6 +429,28 @@ static FBoneWeights CreateWeights(
 IMPLEMENT_SIMPLE_AUTOMATION_TEST(FBoneWeightsTestBlend, "System.AnimationCore.BoneWeights.Blend", EAutomationTestFlags::ApplicationContextMask | EAutomationTestFlags::SmokeFilter)
 bool FBoneWeightsTestBlend::RunTest(const FString& Parameters)
 {
+	{
+		const TCHAR* Name = TEXT("FBoneWeights::Blend[Both Empty]");
+		auto BWA = CreateWeights({});
+		auto BWB = CreateWeights({});
+
+		auto BW = FBoneWeights::Blend(BWA, BWB, 0.5f);
+		UTEST_EQUAL(Name, BW.Num(), 0);
+	}
+	{
+		const TCHAR* Name = TEXT("FBoneWeights::Blend[One Empty]");
+		auto BW0 = CreateWeights({});
+		auto BW1 = CreateWeights({{0, 0.6f}});
+		auto BW2 = CreateWeights({ {1, 0.5f}, {2, 0.25f} });
+
+		auto BWA = FBoneWeights::Blend(BW0, BW1, 0.5f);
+		UTEST_EQUAL(Name, BWA.Num(), 1);
+
+		auto BWB = FBoneWeights::Blend(BW2, BW0, 0.5f);
+		UTEST_EQUAL(Name, BWB.Num(), 2);
+	}
+
+		
 	{
 		const TCHAR* Name = TEXT("FBoneWeights::Blend[One Each, Disjoint]");
 
