Templataize FBoneWeights so that it can be used on any dynamically sizeable container.

#rnx

diff --git a/Engine/Source/Runtime/AnimationCore/Private/BoneWeights.cpp b/Engine/Source/Runtime/AnimationCore/Private/BoneWeights.cpp
--- a/Engine/Source/Runtime/AnimationCore/Private/BoneWeights.cpp
+++ b/Engine/Source/Runtime/AnimationCore/Private/BoneWeights.cpp
@@ -2,10 +2,14 @@
 
 #include "BoneWeights.h"
 
-#include "Algo/IsSorted.h"
 #include "Misc/MemStack.h"
 
 
+namespace UE 
+{
+namespace AnimationCore
+{
+
 static auto WeightSortPredicate = [](const FBoneWeight& A, const FBoneWeight& B) {
 	return A.GetRawWeight() > B.GetRawWeight();
 };
@@ -16,78 +20,8 @@ bool FBoneWeights::SetBoneWeight(
     const FBoneWeightsSettings& InSettings /*= {} */
 )
 {
-	// Does this bone already exist?
-	int32 WeightIndex = FindWeightIndexByBone(InBoneWeight.GetBoneIndex());
-
-	// If the sum of weights could possibly exceed 1.0, we may need normalization based on
-	// the weight settings.
-	bool bMayNeedNormalization = false;
-
-	if (WeightIndex != INDEX_NONE)
-	{
-		FBoneWeight& ExistingBoneWeight = BoneWeights[WeightIndex];
-		// New weight is below the threshold. Remove the current bone weight altogether.
-		if (InBoneWeight.GetRawWeight() < InSettings.GetRawWeightThreshold())
-		{
-			BoneWeights.RemoveAt(WeightIndex);
-
-			// If always normalizing, we need to re-normalize after removing this entry.
-			if (InSettings.GetNormalizeType() == EBoneWeightNormalizeType::Always)
-			{
-				NormalizeWeights(EBoneWeightNormalizeType::Always);
-			}
-
-			return false;
-		}
-
-		if (ExistingBoneWeight.GetRawWeight() == InBoneWeight.GetRawWeight())
-		{
-			return true;
-		}
-
-		bMayNeedNormalization = (ExistingBoneWeight.GetRawWeight() < InBoneWeight.GetRawWeight());
-
-		ExistingBoneWeight.SetRawWeight(InBoneWeight.GetRawWeight());
-	}
-	else
-	{
-		// If the new weight is below the threshold, reject and return.
-		if (InBoneWeight.GetRawWeight() < InSettings.GetRawWeightThreshold())
-		{
-			return false;
-		}
-
-		// Are we already at the limit of weights for this container?
-		if (BoneWeights.Num() == InSettings.GetMaxWeightCount())
-		{
-			// If the weight is smaller than the smallest weight currently, then we reject.
-			if (InBoneWeight.GetRawWeight() < BoneWeights.Last().GetRawWeight())
-			{
-				return false;
-			}
-
-			// Overwrite the last one, we'll put it in its correct place when we sort.
-			BoneWeights.Last() = InBoneWeight;
-		}
-		else
-		{
-			BoneWeights.Add(InBoneWeight);
-		}
-
-		bMayNeedNormalization = true;
-	}
-
-	// If we got here, then we updated/added weights. We're contractually obligated to keep the
-	// weights sorted.
-	SortWeights();
-
-	if ((InSettings.GetNormalizeType() == EBoneWeightNormalizeType::Always) ||
-	    (InSettings.GetNormalizeType() == EBoneWeightNormalizeType::AboveOne && bMayNeedNormalization))
-	{
-		Renormalize(InSettings);
-	}
-
-	return true;
+	FArrayWrapper W(BoneWeights);
+	return W.AddBoneWeight(InBoneWeight, InSettings);
 }
 
 
@@ -96,37 +30,15 @@ bool FBoneWeights::RemoveBoneWeight(
     const FBoneWeightsSettings& InSettings /*= {} */
 )
 {
-	int32 WeightIndex = FindWeightIndexByBone(InBoneIndex);
-	if (WeightIndex == INDEX_NONE)
-	{
-		return false;
-	}
-
-	BoneWeights.RemoveAt(WeightIndex);
-
-	// Cull all weights that exceed limits set by the settings.
-	CullWeights(InSettings);
-
-	// Removing weights will always cause the weight sum to decrease, so we only have to normalize
-	// if always asked to.
-	if (InSettings.GetNormalizeType() == EBoneWeightNormalizeType::Always)
-	{
-		NormalizeWeights(EBoneWeightNormalizeType::Always);
-	}
-
-	return true;
+	FArrayWrapper W(BoneWeights);
+	return W.RemoveBoneWeight(InBoneIndex, InSettings);
 }
 
 
 void FBoneWeights::Renormalize(const FBoneWeightsSettings& InSettings /*= {}*/)
 {
-	NormalizeWeights(InSettings.GetNormalizeType());
-
-	// If entries are now below the threshold, remove them.
-	if (InSettings.GetNormalizeType() == EBoneWeightNormalizeType::Always && CullWeights(InSettings))
-	{
-		NormalizeWeights(EBoneWeightNormalizeType::Always);
-	}
+	FArrayWrapper W(BoneWeights);
+	return W.Renormalize(InSettings);
 }
 
 
@@ -137,27 +49,8 @@ FBoneWeights FBoneWeights::Create(
 )
 {
 	FBoneWeights Result;
-	Result.BoneWeights.Reserve(MaxInlineBoneWeightCount);
-
-	for (int32 Index = 0; Index < MaxInlineBoneWeightCount && InInfluences[Index]; Index++)
-	{
-		FBoneWeight BW(InBones[Index], InInfluences[Index]);
-		if (BW.GetRawWeight() >= InSettings.GetRawWeightThreshold())
-		{
-			Result.BoneWeights.Add(BW);
-		}
-	}
-
-	// Sort the weights by descending weight value before we clip it.
-	Result.SortWeights();
-
-	if (Result.BoneWeights.Num() > InSettings.GetMaxWeightCount())
-	{
-		Result.BoneWeights.SetNum(InSettings.GetMaxWeightCount());
-	}
-
-	Result.NormalizeWeights(InSettings.GetNormalizeType());
-
+	FArrayWrapper W(Result.BoneWeights);
+	W.SetBoneWeights(InBones, InInfluences, InSettings);
 	return Result;
 }
 
@@ -169,54 +62,20 @@ FBoneWeights FBoneWeights::Create(
     const FBoneWeightsSettings& InSettings /*= {} */
 )
 {
-	TArray<FBoneWeight, BoneWeightsTempAllocatorT> StackBoneWeights;
-
-	StackBoneWeights.Reserve(NumEntries);
-	for (int32 Index = 0; Index < NumEntries; Index++)
-	{
-		FBoneWeight BW(InBones[Index], InInfluences[Index]);
-		if (BW.GetRawWeight() >= InSettings.GetRawWeightThreshold())
-		{
-			StackBoneWeights.Add(BW);
-		}
-	}
-
-	return CreateFromArrayView(StackBoneWeights, InSettings);
-}
-
-
-FBoneWeights FBoneWeights::Create(TArrayView<const FBoneWeight> BoneWeights, const FBoneWeightsSettings& InSettings /*= {} */)
-{
-	TArray<FBoneWeight, BoneWeightsTempAllocatorT> StackBoneWeights;
-
-	StackBoneWeights.Reserve(BoneWeights.Num());
-	for (const FBoneWeight& BW : BoneWeights)
-	{
-		if (BW.GetRawWeight() >= InSettings.GetRawWeightThreshold())
-		{
-			StackBoneWeights.Add(BW);
-		}
-	}
-
-	return CreateFromArrayView(StackBoneWeights, InSettings);
+	FBoneWeights Result;
+	FArrayWrapper W(Result.BoneWeights);
+	W.SetBoneWeights(InBones, InInfluences, NumEntries, InSettings);
+	return Result;
 }
 
 
-FBoneWeights FBoneWeights::CreateFromArrayView(
-    TArrayView<FBoneWeight> BoneWeights,
-    const FBoneWeightsSettings& InSettings /*= {} */
+FBoneWeights FBoneWeights::Create(
+	TArrayView<const FBoneWeight> InBoneWeights, 
+	const FBoneWeightsSettings& InSettings /*= {} */
 )
 {
-	BoneWeights.Sort(WeightSortPredicate);
-
-	int32 NumEntries = FMath::Min(BoneWeights.Num(), InSettings.GetMaxWeightCount());
 	FBoneWeights Result;
-	if (NumEntries)
-	{
-		Result.BoneWeights = BoneWeights.Slice(0, NumEntries);
-	}
-	Result.NormalizeWeights(InSettings.GetNormalizeType());
-
+	FArrayWrapper(Result.BoneWeights).SetBoneWeights(InBoneWeights, InSettings);
 	return Result;
 }
 
@@ -228,85 +87,12 @@ FBoneWeights FBoneWeights::Blend(
     const FBoneWeightsSettings& InSettings /*= {} */
 )
 {
-	checkSlow(Algo::IsSorted(InA.BoneWeights, WeightSortPredicate));
-	checkSlow(Algo::IsSorted(InB.BoneWeights, WeightSortPredicate));
-
-	// Both empty?
-	if (InA.BoneWeights.Num() == 0 && InB.BoneWeights.Num() == 0)
-	{
-		return {};
-	}
-	// FIXME: We can probably special-case a few more fast paths. But let's collect statistics
-	// first.
-
-	// To simplify lookup and iteration over the two bone weight arrays, we sort by bone index
-	// value, but indirectly, since we can't sort them directly, as that would violate the
-	// sorted-by-descending-weight contract. Instead we create an indirection array on the stack
-	// and use that to iterate
-	auto CreateIndirectIndex = [](TArrayView<const FBoneWeight> InBoneWeights, TArrayView<int32> InIndexIndirect) {
-		for (int32 Index = 0; Index < InIndexIndirect.Num(); Index++)
-		{
-			InIndexIndirect[Index] = Index;
-		}
-		InIndexIndirect.Sort([InBoneWeights](int32 A, int32 B) {
-			return InBoneWeights[A].GetBoneIndex() < InBoneWeights[B].GetBoneIndex();
-		});
-	};
-
-	TArray<int32, BoneWeightsTempAllocatorT> IndirectIndexA;
-	IndirectIndexA.SetNumUninitialized(InA.BoneWeights.Num());
-	CreateIndirectIndex(InA.BoneWeights, IndirectIndexA);
-
-	TArray<int32, BoneWeightsTempAllocatorT> IndirectIndexB;
-	IndirectIndexB.SetNumUninitialized(InB.BoneWeights.Num());
-	CreateIndirectIndex(InB.BoneWeights, IndirectIndexB);
-
-	TArray<FBoneWeight, TInlineAllocator<MaxInlineBoneWeightCount * 2>> BoneWeights;
-	BoneWeights.Reserve(InA.BoneWeights.Num() + InB.BoneWeights.Num());
-
-	int32 RawBiasB = int32(InBias * float(FBoneWeight::MaxRawWeight));
-	int32 RawBiasA = FBoneWeight::MaxRawWeight - RawBiasB;
-
-	int32 IndexA = 0, IndexB = 0;
-	for (; IndexA < InA.BoneWeights.Num() && IndexB < InB.BoneWeights.Num(); /**/)
-	{
-		const FBoneWeight& BWA = InA.BoneWeights[IndirectIndexA[IndexA]];
-		const FBoneWeight& BWB = InB.BoneWeights[IndirectIndexB[IndexB]];
-
-		// If both have the same bone index, we blend them using the bias given and advance
-		// both arrays. If the bone indices differ, we copy from the array with the lower bone
-		// index value, to ensure we can possibly catch up with the other array. We then
-		// advance until we hit the end of either array after which we blindly copy the remains.
-		if (BWA.GetBoneIndex() == BWB.GetBoneIndex())
-		{
-			uint16 RawWeight = (BWA.GetRawWeight() * RawBiasA + BWB.GetRawWeight() * RawBiasB) / FBoneWeight::MaxRawWeight;
-
-			BoneWeights.Emplace(BWA.GetBoneIndex(), RawWeight);
-			IndexA++;
-			IndexB++;
-		}
-		else if (BWA.GetBoneIndex() < BWB.GetBoneIndex())
-		{
-			BoneWeights.Add(BWA);
-			IndexA++;
-		}
-		else
-		{
-			BoneWeights.Add(BWB);
-			IndexB++;
-		}
-	}
-
-	for (; IndexA < InA.BoneWeights.Num(); IndexA++)
-	{
-		BoneWeights.Add(InA.BoneWeights[IndirectIndexA[IndexA]]);
-	}
-	for (; IndexB < InB.BoneWeights.Num(); IndexB++)
-	{
-		BoneWeights.Add(InB.BoneWeights[IndirectIndexB[IndexB]]);
-	}
-
-	return CreateFromArrayView(BoneWeights, InSettings);
+	FBoneWeights Result;
+	FArrayWrapper W(Result.BoneWeights);
+	FArrayWrapper A(const_cast<FBoneWeights&>(InA).BoneWeights);
+	FArrayWrapper B(const_cast<FBoneWeights&>(InB).BoneWeights);
+	W.Blend(A, B, InBias, InSettings);
+	return Result;
 }
 
 
@@ -381,3 +167,6 @@ void FBoneWeights::NormalizeWeights(
 		}
 	}
 }
+
+} // namespace AnimationCore
+} // namespace UE
diff --git a/Engine/Source/Runtime/AnimationCore/Private/Tests/BoneWeightsTests.cpp b/Engine/Source/Runtime/AnimationCore/Private/Tests/BoneWeightsTests.cpp
--- a/Engine/Source/Runtime/AnimationCore/Private/Tests/BoneWeightsTests.cpp
+++ b/Engine/Source/Runtime/AnimationCore/Private/Tests/BoneWeightsTests.cpp
@@ -12,6 +12,8 @@
 IMPLEMENT_SIMPLE_AUTOMATION_TEST(FBoneWeightTestBasic, "System.AnimationCore.BoneWeight.Basic", EAutomationTestFlags::ApplicationContextMask | EAutomationTestFlags::SmokeFilter)
 bool FBoneWeightTestBasic::RunTest(const FString& Parameters)
 {
+	using namespace UE::AnimationCore;
+
 	// Correct values (uint8)
 	{
 		FBoneWeight BWMin(FBoneIndexType(2), uint8(0));
@@ -75,6 +77,8 @@ bool FBoneWeightTestBasic::RunTest(const FString& Parameters)
 IMPLEMENT_SIMPLE_AUTOMATION_TEST(FBoneWeightsSettingsTest, "System.AnimationCore.BoneWeights.Settings", EAutomationTestFlags::ApplicationContextMask | EAutomationTestFlags::SmokeFilter)
 bool FBoneWeightsSettingsTest::RunTest(const FString& Parameters)
 {
+	using namespace UE::AnimationCore;
+
 	// Ensure that no-one's tampered with the defaults.
 	FBoneWeightsSettings S;
 
@@ -88,6 +92,8 @@ bool FBoneWeightsSettingsTest::RunTest(const FString& Parameters)
 IMPLEMENT_SIMPLE_AUTOMATION_TEST(FBoneWeightsTestBasic, "System.AnimationCore.BoneWeights.Basic", EAutomationTestFlags::ApplicationContextMask | EAutomationTestFlags::SmokeFilter)
 bool FBoneWeightsTestBasic::RunTest(const FString& Parameters)
 {
+	using namespace UE::AnimationCore;
+
 	{
 		const TCHAR* Name = TEXT("FBoneWeights::FBoneWeights");
 
@@ -348,6 +354,8 @@ bool FBoneWeightsTestBasic::RunTest(const FString& Parameters)
 IMPLEMENT_SIMPLE_AUTOMATION_TEST(FBoneWeightsTestCreate, "System.AnimationCore.BoneWeights.Create", EAutomationTestFlags::ApplicationContextMask | EAutomationTestFlags::SmokeFilter)
 bool FBoneWeightsTestCreate::RunTest(const FString& Parameters)
 {
+	using namespace UE::AnimationCore;
+
 	{
 		static const FBoneIndexType Bones[MaxInlineBoneWeightCount] = { 0, 1, 2, 3, 0, 0, 0, 0, 0, 0, 0, 0 };
 		static const uint8 Influences[MaxInlineBoneWeightCount] = { 255, 31, 63, 127, 0, 0, 0, 0, 0, 0, 0, 0 };
@@ -407,10 +415,12 @@ struct FSimpleWeight
 	float Weight;
 };
 
-static FBoneWeights CreateWeights(
+static UE::AnimationCore::FBoneWeights CreateWeights(
 	std::initializer_list<FSimpleWeight> InWeights
 	)
 {
+	using namespace UE::AnimationCore;
+
 	// We only add raw values.
 	FBoneWeightsSettings S;
 	S.SetNormalizeType(EBoneWeightNormalizeType::None);
@@ -429,6 +439,8 @@ static FBoneWeights CreateWeights(
 IMPLEMENT_SIMPLE_AUTOMATION_TEST(FBoneWeightsTestBlend, "System.AnimationCore.BoneWeights.Blend", EAutomationTestFlags::ApplicationContextMask | EAutomationTestFlags::SmokeFilter)
 bool FBoneWeightsTestBlend::RunTest(const FString& Parameters)
 {
+	using namespace UE::AnimationCore;
+
 	{
 		const TCHAR* Name = TEXT("FBoneWeights::Blend[Both Empty]");
 		auto BWA = CreateWeights({});
diff --git a/Engine/Source/Runtime/AnimationCore/Public/BoneWeights.h b/Engine/Source/Runtime/AnimationCore/Public/BoneWeights.h
--- a/Engine/Source/Runtime/AnimationCore/Public/BoneWeights.h
+++ b/Engine/Source/Runtime/AnimationCore/Public/BoneWeights.h
@@ -2,6 +2,7 @@
 
 #pragma once
 
+#include "Algo/IsSorted.h"
 #include "BoneIndices.h"
 #include "Containers/UnrealString.h"
 #include "GPUSkinPublicDefs.h"
@@ -9,6 +10,9 @@
 #include "Serialization/Archive.h"
 #include "Templates/TypeHash.h"
 
+namespace UE { 
+namespace AnimationCore {
+
 /**
 The maximum number of inline bone weights.
 */
@@ -17,17 +21,6 @@ static constexpr int32 MaxInlineBoneWeightCount = MAX_TOTAL_INFLUENCES;
 class FBoneWeight
 {
 public:
-	/**
-	* Default constructor. Deliberately leaves the member values uninitialized for performance.
-	*/
-	FBoneWeight() = default;
-
-	FBoneWeight(const FBoneWeight&) = default;
-	FBoneWeight(FBoneWeight&&) noexcept = default;
-
-	FBoneWeight& operator=(const FBoneWeight&) = default;
-	FBoneWeight& operator=(FBoneWeight&&) = default;
-
 	/**
 	* Returns true if this object's bone weight values are equal to the other container's values.
 	* @note Only equality comparison are supported. Relational comparisons are meaningless.
@@ -49,7 +42,16 @@ public:
 	/**
 	* The maximum raw weight value for a bone influence.
 	*/
-	static constexpr int32 MaxRawWeight = TNumericLimits<uint16>::Max();
+	static constexpr FBoneIndexType MaxRawWeight = TNumericLimits<FBoneIndexType>::Max();
+
+	/** A standard predicate we use for sorting by weight, in a descending order of weights */
+	static bool DescSortByWeightPredicate(const FBoneWeight& A, const FBoneWeight& B)
+	{
+		return A.GetRawWeight() > B.GetRawWeight();
+	}
+	
+	/** Default constructor. NOTE: The values are uninitialized. */
+	FBoneWeight() = default;
 
 	/**
 	* A constructor for old-style bone weights where the weight was stored as an unsigned eight
@@ -130,12 +132,16 @@ public:
 		return RawWeight;
 	}
 
-	friend FArchive& operator<<(FArchive& Ar, FBoneWeight& V)
+	void Serialize(FArchive& InArchive)
 	{
-		Ar << V.BoneIndex;
-		Ar << V.RawWeight;
+		InArchive << BoneIndex;
+		InArchive << RawWeight;
+	}
 
-		return Ar;
+	/** Returns a hash value from the bone weight values */
+	uint32 GetTypeHash() const
+	{
+		return HashCombine(::GetTypeHash(uint16(BoneIndex)), ::GetTypeHash(RawWeight));
 	}
 
 	/**
@@ -147,22 +153,11 @@ public:
 	}
 
 private:
-	friend uint32 GetTypeHash(const FBoneWeight& InBoneWeight);
-
 	FBoneIndexType BoneIndex;
 	uint16 RawWeight;
 };
 
-/**
-* A hashing function to allow the FBoneWeight class to be used with hashing containers (e.g.
-* TSet or TMap).
-*/
-inline uint32 GetTypeHash(const FBoneWeight& InBoneWeight)
-{
-	return HashCombine(GetTypeHash(uint16(InBoneWeight.BoneIndex)),
-	    GetTypeHash(InBoneWeight.RawWeight));
-}
-
+static_assert(sizeof(FBoneWeight) == sizeof(int32), "FBoneWeight must be 32-bits");
 
 /** Specifies the method for which the bone weights are normalized after the bone weight
   * list is modified.
@@ -233,18 +228,268 @@ public:
 		return WeightThreshold;
 	}
 
+	/** Set the default bone index to use if no weights were set. This can be used to ensure
+	    that there's always a valid weight applied to a skinned vertex. */
+	void SetDefaultBoneIndex(FBoneIndexType InBoneIndex)
+	{
+		DefaultBoneIndex = InBoneIndex;
+		bHasDefaultBoneIndex = true;
+	}
+
+	/** Returns the current default bone index to set if no weights were set. If no default
+	    bone index has been set, then this value is undefined. Call HasDefaultBoneIndex to
+		check */
+	FBoneIndexType GetDefaultBoneIndex() const
+	{
+		return DefaultBoneIndex;
+	}
+
+	/** Clears the default bone index. This allows bone weights arrays to be empty if no weights
+		are set. */
+	void ClearDefaultBoneIndex()
+	{
+		bHasDefaultBoneIndex = false;
+	}
+
+	/** Returns \c true if a default bone index should be applied in the absence of other
+	    weights. */
+	bool HasDefaultBoneIndex() const
+	{
+		return bHasDefaultBoneIndex;
+	}
+
 private:
 	EBoneWeightNormalizeType NormalizeType = EBoneWeightNormalizeType::Always;
 	int32 MaxWeightCount = MaxInlineBoneWeightCount;
 	uint16 WeightThreshold = 257; // == uint8(1) converted to 16 bit using 'v | v << 8'.
+	FBoneIndexType DefaultBoneIndex = FBoneIndexType(0);
+	bool bHasDefaultBoneIndex = false;
+};
+
+/** A null adapter for a bone weight container to use with TBoneWeights. Use as a template to
+    create adapters for other types of containers. */
+struct FBoneWeightNullAdapter
+{
+	struct Empty {};
+	using ContainerType = Empty;
+
+	/** Set the number of elements to reserve in the container. The elements can be left 
+	    in an uninitialized state. The TBoneWeights algorithms will ensure that all elements
+		will be properly defined at the end of an operation. */
+	static void SetNum(ContainerType& InContainer, int32 InNum)
+	{
+	}
+
+	static int32 Num(const ContainerType& InContainer)
+	{
+	}
+
+	static FBoneWeight Get(const ContainerType& InContainer, int32 InIndex)
+	{
+		return {};
+	}
+
+	static void Set(ContainerType& InContainer, int32 InIndex, FBoneWeight InBoneWeight)
+	{
+	}
+
+	static void Add(ContainerType& InContainer, FBoneWeight InBoneWeight)
+	{
+	}
+
+	static void Remove(ContainerType& InContainer, int32 InIndex)
+	{
+	}
+
+	template<typename Predicate>
+	static void Sort(ContainerType& InContainer, Predicate InPredicate)
+	{
+	}
+
+	template<typename Predicate>
+	static int32 IndexOf(const ContainerType& InContainer, Predicate InPredicate)
+	{
+		return INDEX_NONE;
+	}
+};
+
+
+/** A templated collection of bone weights algorithms. Requires an adapter to work with a
+    dynamically resizable container. */
+template<typename ContainerAdapter>
+class TBoneWeights
+{
+public:
+	using ContainerType = typename ContainerAdapter::ContainerType;
+
+	TBoneWeights(ContainerType& InContainer) :
+		Container(InContainer)
+	{ }
+
+	inline void SetBoneWeights(
+		TArrayView<const FBoneWeight> BoneWeights,
+		const FBoneWeightsSettings& InSettings = {});
+
+	inline void SetBoneWeights(
+		const FBoneIndexType* InBones,
+		const float* InInfluences,
+		int32 NumEntries,
+		const FBoneWeightsSettings& InSettings = {}
+		);
+
+	inline void SetBoneWeights(
+		const FBoneIndexType InBones[MaxInlineBoneWeightCount],
+		const uint8 InInfluences[MaxInlineBoneWeightCount],
+		const FBoneWeightsSettings& InSettings = {});
+
+	inline bool AddBoneWeight(
+		FBoneWeight InBoneWeight,
+		const FBoneWeightsSettings& InSettings = {}
+		);
+
+	inline bool RemoveBoneWeight(
+		FBoneIndexType InBoneIndex,
+		const FBoneWeightsSettings& InSettings = {}
+		);
+
+	inline void Renormalize(
+		const FBoneWeightsSettings& InSettings = {}
+		);
+
+	/** Blend two bone weights together, making sure to add in every influence from both,
+	  * using the given settings. The bias value should lie on the [0,1] interval. Values outside
+	  * that range may give unwanted results. 
+	  * NOTE: The current container can also be used as an input.
+	  */
+	template<typename ContainerTypeA, typename ContainerTypeB>
+	inline void Blend(
+		const TBoneWeights<ContainerTypeA>& InA,
+		const TBoneWeights<ContainerTypeB>& InB,
+		float InBias,
+		const FBoneWeightsSettings& InSettings = {});
+
+	int32 Num() const
+	{
+		return ContainerAdapter::Num(Container);
+	}
+
+	FBoneWeight Get(int32 Index) const
+	{
+		return ContainerAdapter::Get(Container, Index);
+	}
+
+	FBoneWeight operator[](int32 Index) const
+	{
+		return ContainerAdapter::Get(Container, Index);
+	}
+
+	int32 FindWeightIndexByBone(
+		FBoneIndexType InBoneIndex
+		) const
+	{
+		return Container.IndexOfByPredicate(
+			[InBoneIndex](const FBoneWeight& BW) { return InBoneIndex == BW.GetBoneIndex(); });
+	}
+
+	int32 GetTypeHash() const
+	{
+		uint32 Hash = ::GetTypeHash(Num());
+		for (int32 Index = 0; Index < Num(); Index++)
+		{
+			Hash = HashCombine(Hash, Get(Index).GetTypeHash());
+		}
+		return Hash;
+	}
+
+	FString ToString() const
+	{
+		FString Result(TEXT("["));
+		if (Num())
+		{
+			Result += Get(0).ToString();
+			for (int32 Index = 1; Index < Num(); Index++)
+			{
+				Result += TEXT(", ");
+				Result += Get(Index).ToString();
+			}
+		}
+		Result.Append(TEXT("]"));
+		return Result;
+	}
+
+
+private:
+	inline void SetBoneWeightsInternal(
+		TArrayView<FBoneWeight> BoneWeights,
+		const FBoneWeightsSettings& InSettings = {});
+
+	inline void SortWeights();
+	inline bool CullWeights(const FBoneWeightsSettings& InSettings);
+	inline void NormalizeWeights(EBoneWeightNormalizeType InNormalizeType);
+
+	inline bool Verify() const;
+
+	// The externally owned container we're operating on.
+	ContainerType &Container;
 };
 
+
 /// A simple container for per-vertex influence of bones and their weights.
 class FBoneWeights
 {
 	using BoneWeightsTempAllocatorT = TInlineAllocator<MaxInlineBoneWeightCount>;
 	using BoneWeightArrayT = TArray<FBoneWeight, BoneWeightsTempAllocatorT>;
 
+	template<typename T>
+	struct TArrayAdapter
+	{
+		using ContainerType = T;
+
+		static void SetNum(ContainerType& InContainer, int32 InNum)
+		{
+			const bool bResizeOnShrink = false;
+			InContainer.SetNumUninitialized(InNum, bResizeOnShrink);
+		}
+
+		static int32 Num(const ContainerType& InContainer)
+		{
+			return InContainer.Num();
+		}
+
+		static FBoneWeight Get(const ContainerType& InContainer, int32 InIndex)
+		{
+			return InContainer.GetData()[InIndex];
+		}
+
+		static void Set(ContainerType& InContainer, int32 InIndex, FBoneWeight InBoneWeight)
+		{
+			InContainer.GetData()[InIndex] = InBoneWeight;
+		}
+
+		static void Add(ContainerType& InContainer, FBoneWeight InBoneWeight)
+		{
+			InContainer.Add(InBoneWeight);
+		}
+
+		static void Remove(ContainerType& InContainer, int32 InIndex)
+		{
+			InContainer.RemoveAt(InIndex);
+		}
+
+		template<typename Predicate>
+		static void Sort(ContainerType& InContainer, Predicate InPredicate)
+		{
+			InContainer.Sort(InPredicate);
+		}
+
+		template<typename Predicate>
+		static int32 IndexOf(const ContainerType& InContainer, Predicate InPredicate)
+		{
+			return InContainer.IndexOfByPredicate(InPredicate);
+		}
+	};
+	using FArrayWrapper = TBoneWeights<TArrayAdapter<BoneWeightArrayT>>;
+
 public:
 	FBoneWeights() = default;
 
@@ -358,11 +603,19 @@ public:
 		    [InBoneIndex](const FBoneWeight& BW) { return InBoneIndex == BW.GetBoneIndex(); });
 	}
 
-	friend FArchive& operator<<(FArchive& Ar, FBoneWeights& V)
+	void Serialize(FArchive& InArchive)
 	{
-		Ar << V.BoneWeights;
+		InArchive << BoneWeights;
+	}
 
-		return Ar;
+	int32 GetTypeHash() const
+	{
+		uint32 Hash = ::GetTypeHash(BoneWeights.Num());
+		for (const FBoneWeight& BoneWeight : BoneWeights)
+		{
+			Hash = HashCombine(Hash, BoneWeight.GetTypeHash());
+		}
+		return Hash;
 	}
 
 	FString ToString() const
@@ -388,12 +641,476 @@ private:
 	BoneWeightArrayT BoneWeights;
 };
 
-inline uint32 GetTypeHash(const FBoneWeights& InBoneWeights)
+
+/// TBoneWeights implementation
+template<typename ContainerAdapter>
+void TBoneWeights<ContainerAdapter>::SetBoneWeights(
+	TArrayView<const FBoneWeight> BoneWeights,
+	const FBoneWeightsSettings& InSettings /*= {}*/)
+{
+	TArray<FBoneWeight, TInlineAllocator<MaxInlineBoneWeightCount>> StackBoneWeights;
+
+	StackBoneWeights.Reserve(BoneWeights.Num());
+	for (const FBoneWeight& BW : BoneWeights)
+	{
+		if (BW.GetRawWeight() >= InSettings.GetRawWeightThreshold())
+		{
+			StackBoneWeights.Add(BW);
+		}
+	}
+
+	SetBoneWeightsInternal(StackBoneWeights, InSettings);
+}
+
+
+template<typename ContainerAdapter>
+void TBoneWeights<ContainerAdapter>::SetBoneWeights(
+	const FBoneIndexType* InBones,
+	const float* InInfluences,
+	int32 NumEntries,
+	const FBoneWeightsSettings& InSettings /*= {}*/)
+{
+	TArray<FBoneWeight, TInlineAllocator<MaxInlineBoneWeightCount>> StackBoneWeights;
+
+	StackBoneWeights.Reserve(NumEntries);
+	for (int32 Index = 0; Index < NumEntries; Index++)
+	{
+		FBoneWeight BW(InBones[Index], InInfluences[Index]);
+		if (BW.GetRawWeight() >= InSettings.GetRawWeightThreshold())
+		{
+			StackBoneWeights.Add(BW);
+		}
+	}
+
+	SetBoneWeightsInternal(StackBoneWeights, InSettings);
+}
+
+
+template<typename ContainerAdapter>
+void TBoneWeights<ContainerAdapter>::SetBoneWeights(
+	const FBoneIndexType InBones[MaxInlineBoneWeightCount],
+	const uint8 InInfluences[MaxInlineBoneWeightCount],
+	const FBoneWeightsSettings& InSettings /*= {}*/)
+{
+	// The weights are valid until the first zero influence.
+	int32 NumWeights = 0;
+	for (int32 Index = 0; Index < MaxInlineBoneWeightCount && InInfluences[Index]; Index++)
+	{
+		FBoneWeight BW(InBones[Index], InInfluences[Index]);
+		if (BW.GetRawWeight() >= InSettings.GetRawWeightThreshold())
+		{
+			NumWeights++;
+		}
+	}
+
+	ContainerAdapter::SetNum(Container, NumWeights);
+
+	int32 WeightIndex = 0;
+	for (int32 Index = 0; Index < MaxInlineBoneWeightCount && InInfluences[Index]; Index++)
+	{
+		FBoneWeight BW(InBones[Index], InInfluences[Index]);
+		if (BW.GetRawWeight() >= InSettings.GetRawWeightThreshold())
+		{
+			ContainerAdapter::Set(Container, WeightIndex++, BW);
+		}
+	}
+
+	// Sort the weights by descending weight value before we clip it.
+	SortWeights();
+
+	if (WeightIndex > InSettings.GetMaxWeightCount())
+	{
+		ContainerAdapter::SetNum(Container, InSettings.GetMaxWeightCount());
+	}
+
+	NormalizeWeights(InSettings.GetNormalizeType());
+}
+
+
+template<typename ContainerAdapter>
+void TBoneWeights<ContainerAdapter>::SetBoneWeightsInternal(
+	TArrayView<FBoneWeight> BoneWeights,
+	const FBoneWeightsSettings& InSettings /*= {} */
+	)
+{
+	BoneWeights.Sort(FBoneWeight::DescSortByWeightPredicate);
+
+	int32 NumEntries = FMath::Min(BoneWeights.Num(), InSettings.GetMaxWeightCount());
+	if (NumEntries == 0 && InSettings.HasDefaultBoneIndex())
+	{
+		ContainerAdapter::SetNum(Container, 1);
+		ContainerAdapter::Set(Container, 0, FBoneWeight(InSettings.GetDefaultBoneIndex(), FBoneWeight::MaxRawWeight));
+		return;
+	}
+
+	ContainerAdapter::SetNum(Container, NumEntries);
+
+	for (int32 Index = 0; Index < NumEntries; Index++)
+	{
+		ContainerAdapter::Set(Container, Index, BoneWeights[Index]);
+	}
+	NormalizeWeights(InSettings.GetNormalizeType());
+}
+
+template<typename ContainerAdapter>
+bool TBoneWeights<ContainerAdapter>::AddBoneWeight(
+	FBoneWeight InBoneWeight,
+	const FBoneWeightsSettings& InSettings /*= {}*/)
+{
+	// Does this bone already exist?
+	int32 WeightIndex = FindWeightIndexByBone(InBoneWeight.GetBoneIndex());
+
+	// If the sum of weights could possibly exceed 1.0, we may need normalization based on
+	// the weight settings.
+	bool bMayNeedNormalization = false;
+
+	if (WeightIndex != INDEX_NONE)
+	{
+		FBoneWeight ExistingBoneWeight = ContainerAdapter::Get(Container, WeightIndex);
+
+		// New weight is below the threshold. Remove the current bone weight altogether.
+		if (InBoneWeight.GetRawWeight() < InSettings.GetRawWeightThreshold())
+		{
+			ContainerAdapter::Remove(Container, WeightIndex);
+
+			// If always normalizing, we need to re-normalize after removing this entry.
+			if (InSettings.GetNormalizeType() == EBoneWeightNormalizeType::Always)
+			{
+				NormalizeWeights(EBoneWeightNormalizeType::Always);
+			}
+
+			return false;
+		}
+
+		if (ExistingBoneWeight.GetRawWeight() == InBoneWeight.GetRawWeight())
+		{
+			return true;
+		}
+
+		bMayNeedNormalization = (ExistingBoneWeight.GetRawWeight() < InBoneWeight.GetRawWeight());
+
+		ExistingBoneWeight.SetRawWeight(InBoneWeight.GetRawWeight());
+
+		ContainerAdapter::Set(Container, WeightIndex, ExistingBoneWeight);
+	}
+	else
+	{
+		// If the new weight is below the threshold, reject and return.
+		if (InBoneWeight.GetRawWeight() < InSettings.GetRawWeightThreshold())
+		{
+			return false;
+		}
+
+		// Are we already at the limit of weights for this container?
+		int32 NumWeights = ContainerAdapter::Num(Container);
+		if (NumWeights == InSettings.GetMaxWeightCount())
+		{
+			// If the weight is smaller than the smallest weight currently, then we reject.
+			if (InBoneWeight.GetRawWeight() < ContainerAdapter::Get(Container, NumWeights - 1).GetRawWeight())
+			{
+				return false;
+			}
+
+			// Overwrite the last one, we'll put it in its correct place when we sort.
+			ContainerAdapter::Set(Container, NumWeights - 1, InBoneWeight);
+		}
+		else
+		{
+			ContainerAdapter::Add(Container, InBoneWeight);
+		}
+
+		bMayNeedNormalization = true;
+	}
+
+	// If we got here, then we updated/added weights. We're contractually obligated to keep the
+	// weights sorted.
+	SortWeights();
+
+	if ((InSettings.GetNormalizeType() == EBoneWeightNormalizeType::Always) ||
+		(InSettings.GetNormalizeType() == EBoneWeightNormalizeType::AboveOne && bMayNeedNormalization))
+	{
+		Renormalize(InSettings);
+	}
+
+	return true;
+}
+
+
+template<typename ContainerAdapter>
+bool TBoneWeights<ContainerAdapter>::RemoveBoneWeight(
+    FBoneIndexType InBoneIndex,
+	const FBoneWeightsSettings& InSettings /*= {}*/
+	)
+{
+	int32 WeightIndex = FindWeightIndexByBone(InBoneIndex);
+	if (WeightIndex == INDEX_NONE)
+	{
+		return false;
+	}
+
+	ContainerAdapter::Remove(Container, WeightIndex);
+
+	// Cull all weights that exceed limits set by the settings.
+	CullWeights(InSettings);
+
+	// Removing weights will always cause the weight sum to decrease, so we only have to normalize
+	// if always asked to.
+	if (InSettings.GetNormalizeType() == EBoneWeightNormalizeType::Always)
+	{
+		NormalizeWeights(EBoneWeightNormalizeType::Always);
+	}
+
+	return true;
+}
+
+
+template<typename ContainerAdapter>
+void TBoneWeights<ContainerAdapter>::Renormalize(
+	const FBoneWeightsSettings& InSettings /*= {}*/
+	)
+{
+	NormalizeWeights(InSettings.GetNormalizeType());
+
+	// If entries are now below the threshold, remove them.
+	if (InSettings.GetNormalizeType() == EBoneWeightNormalizeType::Always && CullWeights(InSettings))
+	{
+		NormalizeWeights(EBoneWeightNormalizeType::Always);
+	}
+}
+
+
+template<typename ContainerAdapter>
+template<typename ContainerAdapterA, typename ContainerAdapterB>
+void TBoneWeights<ContainerAdapter>::Blend(
+	const TBoneWeights<ContainerAdapterA>& InBoneWeightsA,
+	const TBoneWeights<ContainerAdapterB>& InBoneWeightsB,
+	float InBias,
+	const FBoneWeightsSettings& InSettings /*= {}*/
+	)
+{
+	checkSlow(InBoneWeightsA.Verify());
+	checkSlow(InBoneWeightsB.Verify());
+
+	// Both empty?
+	if (InBoneWeightsA.Num() == 0 && InBoneWeightsB.Num() == 0)
+	{
+		if (InSettings.HasDefaultBoneIndex())
+		{
+			ContainerAdapter::SetNum(Container, 1);
+			ContainerAdapter::Set(Container, 0, FBoneWeight{InSettings.GetDefaultBoneIndex(), FBoneWeight::MaxRawWeight});
+		}
+		else
+		{
+			ContainerAdapter::SetNum(Container, 0);
+		}		
+		return;
+	}
+	// FIXME: We can probably special-case a few more fast paths (one on either side, one each). 
+	// But let's collect statistics first.
+
+	// To simplify lookup and iteration over the two bone weight arrays, we sort by bone index
+	// value, but indirectly, since we can't sort them directly, as that would violate the
+	// sorted-by-descending-weight contract. Instead we create an indirection array on the stack
+	// and use that to iterate
+	auto CreateIndirectIndex = [](const auto &InBoneWeights, TArrayView<int32> InIndexIndirect) {
+		for (int32 Index = 0; Index < InIndexIndirect.Num(); Index++)
+		{
+			InIndexIndirect[Index] = Index;
+		}
+		InIndexIndirect.Sort([InBoneWeights](int32 A, int32 B) {
+			return InBoneWeights[A].GetBoneIndex() < InBoneWeights[B].GetBoneIndex();
+		});
+	};
+
+	TArray<int32, TInlineAllocator<MaxInlineBoneWeightCount>> IndirectIndexA;
+	IndirectIndexA.SetNumUninitialized(InBoneWeightsA.Num());
+	CreateIndirectIndex(InBoneWeightsA, IndirectIndexA);
+
+	TArray<int32, TInlineAllocator<MaxInlineBoneWeightCount>> IndirectIndexB;
+	IndirectIndexB.SetNumUninitialized(InBoneWeightsB.Num());
+	CreateIndirectIndex(InBoneWeightsB, IndirectIndexB);
+
+	TArray<FBoneWeight, TInlineAllocator<MaxInlineBoneWeightCount * 2>> BoneWeights;
+	BoneWeights.Reserve(InBoneWeightsA.Num() + InBoneWeightsB.Num());
+
+	int32 RawBiasB = int32(InBias * float(FBoneWeight::MaxRawWeight));
+	int32 RawBiasA = FBoneWeight::MaxRawWeight - RawBiasB;
+
+	int32 IndexA = 0, IndexB = 0;
+	for (; IndexA < InBoneWeightsA.Num() && IndexB < InBoneWeightsB.Num(); /**/)
+	{
+		const FBoneWeight& BWA = InBoneWeightsA[IndirectIndexA[IndexA]];
+		const FBoneWeight& BWB = InBoneWeightsB[IndirectIndexB[IndexB]];
+
+		// If both have the same bone index, we blend them using the bias given and advance
+		// both arrays. If the bone indices differ, we copy from the array with the lower bone
+		// index value, to ensure we can possibly catch up with the other array. We then
+		// advance until we hit the end of either array after which we blindly copy the remains.
+		if (BWA.GetBoneIndex() == BWB.GetBoneIndex())
+		{
+			uint16 RawWeight = (BWA.GetRawWeight() * RawBiasA + BWB.GetRawWeight() * RawBiasB) / FBoneWeight::MaxRawWeight;
+
+			BoneWeights.Emplace(BWA.GetBoneIndex(), RawWeight);
+			IndexA++;
+			IndexB++;
+		}
+		else if (BWA.GetBoneIndex() < BWB.GetBoneIndex())
+		{
+			BoneWeights.Add(BWA);
+			IndexA++;
+		}
+		else
+		{
+			BoneWeights.Add(BWB);
+			IndexB++;
+		}
+	}
+
+	for (; IndexA < InBoneWeightsA.Num(); IndexA++)
+	{
+		BoneWeights.Add(InBoneWeightsA[IndirectIndexA[IndexA]]);
+	}
+	for (; IndexB < InBoneWeightsB.Num(); IndexB++)
+	{
+		BoneWeights.Add(InBoneWeightsB[IndirectIndexB[IndexB]]);
+	}
+
+	SetBoneWeightsInternal(BoneWeights, InSettings);
+}
+
+
+template<typename ContainerAdapter>
+void TBoneWeights<ContainerAdapter>::SortWeights()
 {
-	uint32 Hash = GetTypeHash(InBoneWeights.BoneWeights.Num());
-	for (const FBoneWeight& BoneWeight : InBoneWeights)
+	ContainerAdapter::Sort(Container, FBoneWeight::DescSortByWeightPredicate);
+}
+
+
+template<typename ContainerAdapter>
+bool TBoneWeights<ContainerAdapter>::CullWeights(
+	const FBoneWeightsSettings& InSettings
+	)
+{
+	bool bCulled = false;
+	int32 NumWeights = ContainerAdapter::Num(Container);
+
+	// If are are more entries in the container than the settings allow for, indiscriminately 
+	// remove the excess entries.
+	if (NumWeights > InSettings.GetMaxWeightCount())
+	{
+		ContainerAdapter::SetNum(Container, InSettings.GetMaxWeightCount());
+		NumWeights = InSettings.GetMaxWeightCount();
+		bCulled = true;
+	}
+
+	// If any remaining entries are now below the threshold, remove them too.
+	while (NumWeights > 0 && ContainerAdapter::Get(Container, NumWeights - 1).GetRawWeight() < InSettings.GetRawWeightThreshold())
+	{
+		ContainerAdapter::SetNum(Container, --NumWeights);
+		bCulled = true;
+	}
+
+	return bCulled;
+}
+
+
+template<typename ContainerAdapter>
+void TBoneWeights<ContainerAdapter>::NormalizeWeights(
+	EBoneWeightNormalizeType InNormalizeType
+	)
+{
+	int32 NumWeights = ContainerAdapter::Num(Container);
+
+	// Early checks
+	if (InNormalizeType == EBoneWeightNormalizeType::None || NumWeights == 0)
+	{
+		return;
+	}
+
+	// Common case.
+	if (NumWeights == 1)
+	{
+		if (InNormalizeType == EBoneWeightNormalizeType::Always)
+		{
+			// Set the weight to full for the sole entry if normalizing always.
+			FBoneWeight BoneWeight = ContainerAdapter::Get(Container, 0);
+			BoneWeight.SetRawWeight(FBoneWeight::MaxRawWeight);
+			ContainerAdapter::Set(Container, 0, BoneWeight);
+		}
+		return;
+	}
+
+	// We operate on int64, since we can easily end up with wraparound issues during one of the
+	// multiplications below when using int32. This would tank the division by WeightSum.
+	int64 WeightSum = 0;
+	for (int32 Index = 0; Index < NumWeights; Index++)
+	{
+		WeightSum += ContainerAdapter::Get(Container, Index).GetRawWeight();
+	}
+
+	if ((InNormalizeType == EBoneWeightNormalizeType::Always && ensure(WeightSum != 0)) ||
+		WeightSum > FBoneWeight::MaxRawWeight)
 	{
-		Hash = HashCombine(Hash, GetTypeHash(BoneWeight));
+		int64 Correction = 0;
+
+		// Here we treat the raw weight as a 16.16 fixed point value and ensure that the
+		// fraction, which would otherwise be lost through rounding, is carried over to the
+		// subsequent values to maintain a constant sum to the max weight value.
+		// We do this in descending weight order in an attempt to ensure that weight values
+		// aren't needlessly lost after scaling.
+		for (int32 Index = 0; Index < NumWeights; Index++)
+		{
+			FBoneWeight BW = ContainerAdapter::Get(Container, Index);
+			int64 ScaledWeight = int64(BW.GetRawWeight()) * FBoneWeight::MaxRawWeight + Correction;
+			BW.SetRawWeight(uint16(FMath::Min(ScaledWeight / WeightSum, int64(FBoneWeight::MaxRawWeight))));
+			Correction = ScaledWeight - BW.GetRawWeight() * WeightSum;
+			ContainerAdapter::Set(Container, Index, BW);
+		}
 	}
-	return Hash;
+}
+
+
+template<typename ContainerAdapter>
+bool UE::AnimationCore::TBoneWeights<ContainerAdapter>::Verify() const
+{
+	// FIXME: TBD.
+	return true;
+}
+
+
+} // namespace AnimationCore
+} // namespace UE
+
+
+/**
+* A hashing function to allow the FBoneWeight class to be used with hashing containers (e.g.
+* TSet or TMap).
+*/
+static inline uint32 GetTypeHash(
+	const UE::AnimationCore::FBoneWeight& InBoneWeight
+	)
+{
+	return InBoneWeight.GetTypeHash();
+}
+
+static inline FArchive& operator<<(
+	FArchive& InArchive, 
+	UE::AnimationCore::FBoneWeight& InOutBoneWeight
+	)
+{
+	InOutBoneWeight.Serialize(InArchive);
+	return InArchive;
+}
+
+static inline uint32 GetTypeHash(const UE::AnimationCore::FBoneWeights& InBoneWeights)
+{
+	return InBoneWeights.GetTypeHash();
+}
+
+static inline FArchive& operator<<(
+    FArchive& InArchive,
+    UE::AnimationCore::FBoneWeights& InOutBoneWeights)
+{
+	InOutBoneWeights.Serialize(InArchive);
+	return InArchive;
 }
