Missing return value in the dummy TBoneWeights adapter.

#trivial
#rnx

diff --git a/Engine/Source/Runtime/AnimationCore/Public/BoneWeights.h b/Engine/Source/Runtime/AnimationCore/Public/BoneWeights.h
--- a/Engine/Source/Runtime/AnimationCore/Public/BoneWeights.h
+++ b/Engine/Source/Runtime/AnimationCore/Public/BoneWeights.h
@@ -282,6 +282,7 @@ struct FBoneWeightNullAdapter
 
 	static int32 Num(const ContainerType& InContainer)
 	{
+		return INDEX_NONE;
 	}
 
 	static FBoneWeight Get(const ContainerType& InContainer, int32 InIndex)
