Change FBoneWeight::MaxRawWeight to a constexpr function to satisfy Mac linkage issues.

#trivial
#rnx

diff --git a/Engine/Source/Runtime/AnimationCore/Private/BoneWeights.cpp b/Engine/Source/Runtime/AnimationCore/Private/BoneWeights.cpp
--- a/Engine/Source/Runtime/AnimationCore/Private/BoneWeights.cpp
+++ b/Engine/Source/Runtime/AnimationCore/Private/BoneWeights.cpp
@@ -136,7 +136,7 @@ void FBoneWeights::NormalizeWeights(
 	{
 		if (InNormalizeType == EBoneWeightNormalizeType::Always)
 		{
-			BoneWeights[0].SetRawWeight(FBoneWeight::MaxRawWeight);
+			BoneWeights[0].SetRawWeight(FBoneWeight::GetMaxRawWeight());
 		}
 		return;
 	}
@@ -150,7 +150,7 @@ void FBoneWeights::NormalizeWeights(
 	}
 
 	if ((InNormalizeType == EBoneWeightNormalizeType::Always && ensure(WeightSum != 0)) ||
-	    WeightSum > FBoneWeight::MaxRawWeight)
+	    WeightSum > FBoneWeight::GetMaxRawWeight())
 	{
 		int64 Correction = 0;
 
@@ -161,8 +161,8 @@ void FBoneWeights::NormalizeWeights(
 		// aren't needlessly lost after scaling.
 		for (FBoneWeight& BW : BoneWeights)
 		{
-			int64 ScaledWeight = int64(BW.GetRawWeight()) * FBoneWeight::MaxRawWeight + Correction;
-			BW.SetRawWeight(uint16(FMath::Min(ScaledWeight / WeightSum, int64(FBoneWeight::MaxRawWeight))));
+			int64 ScaledWeight = int64(BW.GetRawWeight()) * FBoneWeight::GetMaxRawWeight() + Correction;
+			BW.SetRawWeight(uint16(FMath::Min(ScaledWeight / WeightSum, int64(FBoneWeight::GetMaxRawWeight()))));
 			Correction = ScaledWeight - BW.GetRawWeight() * WeightSum;
 		}
 	}
diff --git a/Engine/Source/Runtime/AnimationCore/Private/Tests/BoneWeightsTests.cpp b/Engine/Source/Runtime/AnimationCore/Private/Tests/BoneWeightsTests.cpp
--- a/Engine/Source/Runtime/AnimationCore/Private/Tests/BoneWeightsTests.cpp
+++ b/Engine/Source/Runtime/AnimationCore/Private/Tests/BoneWeightsTests.cpp
@@ -30,7 +30,7 @@ bool FBoneWeightTestBasic::RunTest(const FString& Parameters)
 
 		FBoneWeight BWMax(FBoneIndexType(0), uint8(255));
 		UTEST_EQUAL(TEXT("FBoneWeight max value (uint8)"), BWMax.GetWeight(), 1.0f);
-		UTEST_EQUAL(TEXT("FBoneWeight max raw value (uint8)"), BWMax.GetRawWeight(), FBoneWeight::MaxRawWeight);
+		UTEST_EQUAL(TEXT("FBoneWeight max raw value (uint8)"), BWMax.GetRawWeight(), FBoneWeight::GetMaxRawWeight());
 	}
 
 	// Correct values (uint16)
@@ -42,12 +42,12 @@ bool FBoneWeightTestBasic::RunTest(const FString& Parameters)
 
 		// Since uint16::max / 2 -is equal to 32767.5
 		FBoneWeight BWMid(FBoneIndexType(0), uint16(0x7FFF));
-		UTEST_EQUAL(TEXT("FBoneWeight mid value (uint16)"), BWMid.GetWeight(), 0x7FFF / float(FBoneWeight::MaxRawWeight));
+		UTEST_EQUAL(TEXT("FBoneWeight mid value (uint16)"), BWMid.GetWeight(), 0x7FFF / float(FBoneWeight::GetMaxRawWeight()));
 		UTEST_EQUAL(TEXT("FBoneWeight mid raw value (uint16)"), BWMid.GetRawWeight(), 0x7FFF);
 
-		FBoneWeight BWMax(FBoneIndexType(0), uint16(FBoneWeight::MaxRawWeight));
+		FBoneWeight BWMax(FBoneIndexType(0), uint16(FBoneWeight::GetMaxRawWeight()));
 		UTEST_EQUAL(TEXT("FBoneWeight max value (uint16)"), BWMax.GetWeight(), 1.0f);
-		UTEST_EQUAL(TEXT("FBoneWeight max raw value (uint16)"), BWMax.GetRawWeight(), FBoneWeight::MaxRawWeight);
+		UTEST_EQUAL(TEXT("FBoneWeight max raw value (uint16)"), BWMax.GetRawWeight(), FBoneWeight::GetMaxRawWeight());
 	}
 
 	// Copy constructor
@@ -125,7 +125,7 @@ bool FBoneWeightsTestBasic::RunTest(const FString& Parameters)
 
 		FBoneWeight BI = A[0];
 		UTEST_EQUAL(Name, BI.GetWeight(), 1.0f);
-		UTEST_EQUAL(Name, BI.GetRawWeight(), FBoneWeight::MaxRawWeight);
+		UTEST_EQUAL(Name, BI.GetRawWeight(), FBoneWeight::GetMaxRawWeight());
 	}
 
 	// Check normalization works
@@ -205,7 +205,7 @@ bool FBoneWeightsTestBasic::RunTest(const FString& Parameters)
 		FBoneWeights A;
 		for (uint32 Index = 0; Index < MaxInlineBoneWeightCount; Index++)
 		{
-			A.SetBoneWeight(FBoneIndexType(Index), Index == 0 ? 1.0f : 2.0f / float(FBoneWeight::MaxRawWeight), SN);
+			A.SetBoneWeight(FBoneIndexType(Index), Index == 0 ? 1.0f : 2.0f / float(FBoneWeight::GetMaxRawWeight()), SN);
 		}
 
 		A.Renormalize(SA);
@@ -225,7 +225,7 @@ bool FBoneWeightsTestBasic::RunTest(const FString& Parameters)
 				UTEST_EQUAL(NameSub, int32(A[Index].GetRawWeight()), 2);
 			}
 		}
-		UTEST_EQUAL(Name, WeightSum, int32(FBoneWeight::MaxRawWeight));
+		UTEST_EQUAL(Name, WeightSum, int32(FBoneWeight::GetMaxRawWeight()));
 
 
 	}
diff --git a/Engine/Source/Runtime/AnimationCore/Public/BoneWeights.h b/Engine/Source/Runtime/AnimationCore/Public/BoneWeights.h
--- a/Engine/Source/Runtime/AnimationCore/Public/BoneWeights.h
+++ b/Engine/Source/Runtime/AnimationCore/Public/BoneWeights.h
@@ -42,7 +42,10 @@ public:
 	/**
 	* The maximum raw weight value for a bone influence.
 	*/
-	static constexpr FBoneIndexType MaxRawWeight = TNumericLimits<FBoneIndexType>::Max();
+	static constexpr FBoneIndexType GetMaxRawWeight() 
+	{
+		return TNumericLimits<FBoneIndexType>::Max();
+	}
 
 	/** A standard predicate we use for sorting by weight, in a descending order of weights */
 	static bool DescSortByWeightPredicate(const FBoneWeight& A, const FBoneWeight& B)
@@ -104,7 +107,7 @@ public:
 	void SetWeight(float InWeight)
 	{
 		InWeight = FMath::Clamp(InWeight, 0.0f, 1.0f);
-		RawWeight = uint16(InWeight * float(MaxRawWeight) + 0.5f);
+		RawWeight = uint16(InWeight * float(GetMaxRawWeight()) + 0.5f);
 	}
 
 	/**
@@ -112,7 +115,7 @@ public:
 	*/
 	float GetWeight() const
 	{
-		return RawWeight / float(MaxRawWeight);
+		return RawWeight / float(GetMaxRawWeight());
 	}
 
 	/**
@@ -209,7 +212,7 @@ public:
 	FBoneWeightsSettings& SetWeightThreshold(float InWeightThreshold)
 	{
 		InWeightThreshold = FMath::Clamp(InWeightThreshold, 0.0f, 1.0f);
-		WeightThreshold = uint16(InWeightThreshold * FBoneWeight::MaxRawWeight + 0.5f);
+		WeightThreshold = uint16(InWeightThreshold * FBoneWeight::GetMaxRawWeight() + 0.5f);
 		WeightThreshold = FMath::Max(WeightThreshold, uint16(1));
 		return *this;
 	}
@@ -740,7 +743,7 @@ void TBoneWeights<ContainerAdapter>::SetBoneWeightsInternal(
 	if (NumEntries == 0 && InSettings.HasDefaultBoneIndex())
 	{
 		ContainerAdapter::SetNum(Container, 1);
-		ContainerAdapter::Set(Container, 0, FBoneWeight(InSettings.GetDefaultBoneIndex(), FBoneWeight::MaxRawWeight));
+		ContainerAdapter::Set(Container, 0, FBoneWeight(InSettings.GetDefaultBoneIndex(), FBoneWeight::GetMaxRawWeight()));
 		return;
 	}
 
@@ -898,7 +901,7 @@ void TBoneWeights<ContainerAdapter>::Blend(
 		if (InSettings.HasDefaultBoneIndex())
 		{
 			ContainerAdapter::SetNum(Container, 1);
-			ContainerAdapter::Set(Container, 0, FBoneWeight{InSettings.GetDefaultBoneIndex(), FBoneWeight::MaxRawWeight});
+			ContainerAdapter::Set(Container, 0, FBoneWeight{InSettings.GetDefaultBoneIndex(), FBoneWeight::GetMaxRawWeight()});
 		}
 		else
 		{
@@ -934,8 +937,8 @@ void TBoneWeights<ContainerAdapter>::Blend(
 	TArray<FBoneWeight, TInlineAllocator<MaxInlineBoneWeightCount * 2>> BoneWeights;
 	BoneWeights.Reserve(InBoneWeightsA.Num() + InBoneWeightsB.Num());
 
-	int32 RawBiasB = int32(InBias * float(FBoneWeight::MaxRawWeight));
-	int32 RawBiasA = FBoneWeight::MaxRawWeight - RawBiasB;
+	int32 RawBiasB = int32(InBias * float(FBoneWeight::GetMaxRawWeight()));
+	int32 RawBiasA = FBoneWeight::GetMaxRawWeight() - RawBiasB;
 
 	int32 IndexA = 0, IndexB = 0;
 	for (; IndexA < InBoneWeightsA.Num() && IndexB < InBoneWeightsB.Num(); /**/)
@@ -949,7 +952,7 @@ void TBoneWeights<ContainerAdapter>::Blend(
 		// advance until we hit the end of either array after which we blindly copy the remains.
 		if (BWA.GetBoneIndex() == BWB.GetBoneIndex())
 		{
-			uint16 RawWeight = (BWA.GetRawWeight() * RawBiasA + BWB.GetRawWeight() * RawBiasB) / FBoneWeight::MaxRawWeight;
+			uint16 RawWeight = (BWA.GetRawWeight() * RawBiasA + BWB.GetRawWeight() * RawBiasB) / FBoneWeight::GetMaxRawWeight();
 
 			BoneWeights.Emplace(BWA.GetBoneIndex(), RawWeight);
 			IndexA++;
@@ -1035,7 +1038,7 @@ void TBoneWeights<ContainerAdapter>::NormalizeWeights(
 		{
 			// Set the weight to full for the sole entry if normalizing always.
 			FBoneWeight BoneWeight = ContainerAdapter::Get(Container, 0);
-			BoneWeight.SetRawWeight(FBoneWeight::MaxRawWeight);
+			BoneWeight.SetRawWeight(FBoneWeight::GetMaxRawWeight());
 			ContainerAdapter::Set(Container, 0, BoneWeight);
 		}
 		return;
@@ -1050,7 +1053,7 @@ void TBoneWeights<ContainerAdapter>::NormalizeWeights(
 	}
 
 	if ((InNormalizeType == EBoneWeightNormalizeType::Always && ensure(WeightSum != 0)) ||
-		WeightSum > FBoneWeight::MaxRawWeight)
+		WeightSum > FBoneWeight::GetMaxRawWeight())
 	{
 		int64 Correction = 0;
 
@@ -1062,8 +1065,8 @@ void TBoneWeights<ContainerAdapter>::NormalizeWeights(
 		for (int32 Index = 0; Index < NumWeights; Index++)
 		{
 			FBoneWeight BW = ContainerAdapter::Get(Container, Index);
-			int64 ScaledWeight = int64(BW.GetRawWeight()) * FBoneWeight::MaxRawWeight + Correction;
-			BW.SetRawWeight(uint16(FMath::Min(ScaledWeight / WeightSum, int64(FBoneWeight::MaxRawWeight))));
+			int64 ScaledWeight = int64(BW.GetRawWeight()) * FBoneWeight::GetMaxRawWeight() + Correction;
+			BW.SetRawWeight(uint16(FMath::Min(ScaledWeight / WeightSum, int64(FBoneWeight::GetMaxRawWeight()))));
 			Correction = ScaledWeight - BW.GetRawWeight() * WeightSum;
 			ContainerAdapter::Set(Container, Index, BW);
 		}
