Templatized FBoneWeight into TBoneWeight that can sit on top of an arbitrary array-like container.
Register new SkinWeights mesh descriptor attribute that combines weights and influence in a single attribute, that proxies TBoneWeight to allow setting bone weights in one shot with normalization, culling, etc.
Deprecated and removed InfluenceBones and InfluenceWeights and updated the calling code to use the new attribute.

#rnx

diff --git a/Engine/Plugins/Runtime/Interchange/Source/Import/Private/Mesh/InterchangeSkeletalMeshFactory.cpp b/Engine/Plugins/Runtime/Interchange/Source/Import/Private/Mesh/InterchangeSkeletalMeshFactory.cpp
--- a/Engine/Plugins/Runtime/Interchange/Source/Import/Private/Mesh/InterchangeSkeletalMeshFactory.cpp
+++ b/Engine/Plugins/Runtime/Interchange/Source/Import/Private/Mesh/InterchangeSkeletalMeshFactory.cpp
@@ -600,8 +600,7 @@ namespace UE
 
 				FSkeletalMeshConstAttributes Attributes(SourceMeshDescription);
 				TVertexAttributesConstRef<FVector> VertexPositions = Attributes.GetVertexPositions();
-				TVertexAttributesConstRef<TArrayAttribute<int32>> VertexInfluenceBones = Attributes.GetVertexInfluenceBones();
-				TVertexAttributesConstRef<TArrayAttribute<float>> VertexInfluenceWeights = Attributes.GetVertexInfluenceWeights();
+				FSkinWeightsVertexAttributesConstRef VertexSkinWeights = Attributes.GetVertexSkinWeights();
 
 				TVertexInstanceAttributesConstRef<FVertexID> VertexInstanceVertexIndices = Attributes.GetVertexInstanceVertexIndices();
 				TVertexInstanceAttributesConstRef<FVector2D> VertexInstanceUVs = Attributes.GetVertexInstanceUVs();
@@ -639,15 +638,18 @@ namespace UE
 				{
 					//We can use GetValue because the Meshdescription was compacted before the copy
 					DestinationSkeletalMeshImportData.Points[VertexID.GetValue()] = VertexPositions[VertexID];
-					int32 InfluenceCount = VertexInfluenceBones[VertexID].Num();
-					int32 InfluenceOffsetIndex = DestinationSkeletalMeshImportData.Influences.Num();
+
+					FVertexBoneWeightsConst BoneWeights = VertexSkinWeights.Get(VertexID); 
+					const int32 InfluenceCount = BoneWeights.Num();
+					
+					const int32 InfluenceOffsetIndex = DestinationSkeletalMeshImportData.Influences.Num();
 					DestinationSkeletalMeshImportData.Influences.AddDefaulted(InfluenceCount);
 					for (int32 InfluenceIndex = 0; InfluenceIndex < InfluenceCount; ++InfluenceIndex)
 					{
 						SkeletalMeshImportData::FRawBoneInfluence& BoneInfluence = DestinationSkeletalMeshImportData.Influences[InfluenceOffsetIndex + InfluenceIndex];
 						BoneInfluence.VertexIndex = VertexID.GetValue();
-						BoneInfluence.BoneIndex = VertexInfluenceBones[VertexID][InfluenceIndex];
-						BoneInfluence.Weight = VertexInfluenceWeights[VertexID][InfluenceIndex];
+						BoneInfluence.BoneIndex = BoneWeights[InfluenceIndex].GetBoneIndex();
+						BoneInfluence.Weight = BoneWeights[InfluenceIndex].GetWeight();
 					}
 				}
 
diff --git a/Engine/Plugins/Runtime/Interchange/Source/Parsers/Fbx/Private/FbxMesh.cpp b/Engine/Plugins/Runtime/Interchange/Source/Parsers/Fbx/Private/FbxMesh.cpp
--- a/Engine/Plugins/Runtime/Interchange/Source/Parsers/Fbx/Private/FbxMesh.cpp
+++ b/Engine/Plugins/Runtime/Interchange/Source/Parsers/Fbx/Private/FbxMesh.cpp
@@ -1009,38 +1009,13 @@ namespace UE
 						}
 
 						FSkeletalMeshAttributes SkeletalMeshAttributes(*MeshDescription);
+						SkeletalMeshAttributes.Register();
 
+						using namespace UE::AnimationCore;
+						TMap<FVertexID, TArray<FBoneWeight>> RawBoneWeights;
+						
 						//Add the influence data in the skeletalmesh description
-						TVertexAttributesRef<TArrayAttribute<int32>> VertexInfluenceBones = SkeletalMeshAttributes.GetVertexInfluenceBones();
-						TVertexAttributesRef<TArrayAttribute<float>> VertexInfluenceWeights = SkeletalMeshAttributes.GetVertexInfluenceWeights();
-
-
-						auto AddOneInfluence = [&VertexInfluenceBones, &VertexInfluenceWeights](FVertexID VertexID, int32 BoneIndex, float Weight)
-						{
-							//Add one influence to the vertex if its not already set
-							int32 InfluenceCount = VertexInfluenceBones[VertexID].Num();
-							bool bCreateNewInfluenceEntry = true;
-							for (int32 InfluenceIndex = 0; InfluenceIndex < InfluenceCount; ++InfluenceIndex)
-							{
-								int32 InfluenceBoneIndex = VertexInfluenceBones[VertexID][InfluenceIndex];
-								if (InfluenceBoneIndex == BoneIndex)
-								{
-									//Use the most heavy weight
-									float InfluenceWeight = VertexInfluenceWeights[VertexID][InfluenceIndex];
-									if (Weight > InfluenceWeight)
-									{
-										VertexInfluenceWeights[VertexID][InfluenceIndex] = Weight;
-									}
-									bCreateNewInfluenceEntry = false;
-									break;
-								}
-							}
-							if (bCreateNewInfluenceEntry)
-							{
-								ensure(InfluenceCount == VertexInfluenceBones[VertexID].Add(BoneIndex));
-								ensure(InfluenceCount == VertexInfluenceWeights[VertexID].Add(Weight));
-							}
-						};
+						FSkinWeightsVertexAttributesRef VertexSkinWeights = SkeletalMeshAttributes.GetVertexSkinWeights();
 
 						if (MeshType == EMeshType::Skinned)
 						{
@@ -1083,9 +1058,41 @@ namespace UE
 									}
 									float Weight = static_cast<float>(Weights[ControlPointIndex]);
 
-									AddOneInfluence(VertexID, BoneIndex, Weight);
+									TArray<FBoneWeight> *BoneWeights = RawBoneWeights.Find(VertexID);
+									if (BoneWeights)
+									{
+										// Do we already have a weight for this bone? Keep the greater weight.
+										bool bShouldAdd = true;
+										for (int32 WeightIndex = 0; WeightIndex < BoneWeights->Num(); WeightIndex++)
+										{
+											FBoneWeight &BoneWeight = (*BoneWeights)[WeightIndex];
+											if (BoneWeight.GetBoneIndex() == BoneIndex)
+											{
+												if (BoneWeight.GetWeight() < Weight)
+												{
+													BoneWeight.SetWeight(Weight);
+												}
+												bShouldAdd = false;
+												break;
+											}
+										}
+										if (bShouldAdd)
+										{
+											BoneWeights->Add(FBoneWeight(BoneIndex, Weight));
+										}
+									}
+									else
+									{
+										RawBoneWeights.Add(VertexID).Add(FBoneWeight(BoneIndex, Weight));
+									}
 								}
 							}
+
+							// Add all the raw bone weights. This will cause the weights to be sorted and re-normalized after culling to max influences.
+							for (const TTuple<FVertexID, TArray<FBoneWeight>> &Item: RawBoneWeights)
+							{
+								VertexSkinWeights.Set(Item.Key, Item.Value);
+							}	
 						}
 						else // for rigid mesh
 						{
@@ -1093,6 +1100,8 @@ namespace UE
 							int32 BoneIndex = -1;
 							SortedJoints->Find(MeshNode, BoneIndex);
 
+							const float Weight = 1.0f;
+							const FBoneWeights DefaultBinding = FBoneWeights::Create({FBoneWeight(BoneIndex, Weight)});
 							if (BoneIndex == -1 && SortedJoints->Num() > 0)
 							{
 								//When we import geometry only, we want to hook all the influence to a particular bone
@@ -1107,73 +1116,11 @@ namespace UE
 									//Invalid Influence
 									continue;
 								}
-								const float Weight = 1.0f;
-								AddOneInfluence(VertexID, BoneIndex, Weight);
-							}
-						}
-
-						//Reorder by weight
-						{
-							TArray<int32> BoneIndexes;
-							TArray<float> BoneWeights;
-							TArray<int32> RemapIndexes;
-							TArray<float> RemapWeights;
-							for (const FVertexID VertexID : MeshDescription->Vertices().GetElementIDs())
-							{
-								//Add one influence to the vertex if its not already set
-								int32 InfluenceCount = VertexInfluenceBones[VertexID].Num();
-								BoneIndexes.Reset(InfluenceCount);
-								BoneIndexes.AddZeroed(InfluenceCount);
-								BoneWeights.Reset(InfluenceCount);
-								BoneWeights.AddZeroed(InfluenceCount);
-								RemapIndexes.Reset(InfluenceCount);
-								RemapIndexes.AddZeroed(InfluenceCount);
-								RemapWeights.Reset(InfluenceCount);
-								RemapWeights.AddZeroed(InfluenceCount);
-								for (int32 InfluenceIndex = 0; InfluenceIndex < InfluenceCount; ++InfluenceIndex)
-								{
-									int32 InfluenceBoneIndex = VertexInfluenceBones[VertexID][InfluenceIndex];
-									BoneIndexes[InfluenceIndex] = InfluenceBoneIndex;
-									//Use the most heavy weight
-									float InfluenceWeight = VertexInfluenceWeights[VertexID][InfluenceIndex];
-									BoneWeights[InfluenceIndex] = InfluenceWeight;
-									RemapIndexes[InfluenceIndex] = InfluenceIndex;
-									RemapWeights[InfluenceIndex] = InfluenceWeight;
-								}
 
-								for (int32 InfluenceIndex = 0; InfluenceIndex < InfluenceCount; ++InfluenceIndex)
-								{
-									if (InfluenceIndex == 0)
-									{
-										continue;
-									}
-									float CurrentWeight = BoneWeights[InfluenceIndex];
-									for (int32 RemapIndex = InfluenceIndex - 1; RemapIndex >= 0; RemapIndex--)
-									{
-										if (CurrentWeight > BoneWeights[RemapIndexes[RemapIndex]])
-										{
-											//Swap
-											int32 SwapIndex = RemapIndexes[RemapIndex + 1];
-											float SwapValue = RemapWeights[RemapIndex + 1];
-											RemapIndexes[RemapIndex + 1] = RemapIndexes[RemapIndex];
-											RemapWeights[RemapIndex + 1] = RemapWeights[RemapIndex];
-											RemapIndexes[RemapIndex] = SwapIndex;
-											RemapWeights[RemapIndex] = SwapValue;
-										}
-										else
-										{
-											break;
-										}
-									}
-								}
-								//Push the value into the mesh description
-								for (int32 InfluenceIndex = 0; InfluenceIndex < InfluenceCount; ++InfluenceIndex)
-								{
-									BoneIndexes[InfluenceIndex] = RemapIndexes[InfluenceIndex];
-									BoneWeights[InfluenceIndex] = RemapWeights[InfluenceIndex];
-								}
+								VertexSkinWeights.Set(VertexID, DefaultBinding);
 							}
 						}
+
 					}
 				}
 				// needed?
diff --git a/Engine/Source/Runtime/AnimationCore/Private/BoneWeights.cpp b/Engine/Source/Runtime/AnimationCore/Private/BoneWeights.cpp
--- a/Engine/Source/Runtime/AnimationCore/Private/BoneWeights.cpp
+++ b/Engine/Source/Runtime/AnimationCore/Private/BoneWeights.cpp
@@ -5,168 +5,3 @@
 #include "Misc/MemStack.h"
 
 
-namespace UE 
-{
-namespace AnimationCore
-{
-
-static auto WeightSortPredicate = [](const FBoneWeight& A, const FBoneWeight& B) {
-	return A.GetRawWeight() > B.GetRawWeight();
-};
-
-
-bool FBoneWeights::SetBoneWeight(
-    FBoneWeight InBoneWeight,
-    const FBoneWeightsSettings& InSettings /*= {} */
-)
-{
-	FArrayWrapper W(BoneWeights);
-	return W.AddBoneWeight(InBoneWeight, InSettings);
-}
-
-
-bool FBoneWeights::RemoveBoneWeight(
-    FBoneIndexType InBoneIndex,
-    const FBoneWeightsSettings& InSettings /*= {} */
-)
-{
-	FArrayWrapper W(BoneWeights);
-	return W.RemoveBoneWeight(InBoneIndex, InSettings);
-}
-
-
-void FBoneWeights::Renormalize(const FBoneWeightsSettings& InSettings /*= {}*/)
-{
-	FArrayWrapper W(BoneWeights);
-	return W.Renormalize(InSettings);
-}
-
-
-FBoneWeights FBoneWeights::Create(
-    const FBoneIndexType InBones[MaxInlineBoneWeightCount],
-    const uint8 InInfluences[MaxInlineBoneWeightCount],
-    const FBoneWeightsSettings& InSettings /*= {} */
-)
-{
-	FBoneWeights Result;
-	FArrayWrapper W(Result.BoneWeights);
-	W.SetBoneWeights(InBones, InInfluences, InSettings);
-	return Result;
-}
-
-
-FBoneWeights FBoneWeights::Create(
-    const FBoneIndexType* InBones,
-    const float* InInfluences,
-    int32 NumEntries,
-    const FBoneWeightsSettings& InSettings /*= {} */
-)
-{
-	FBoneWeights Result;
-	FArrayWrapper W(Result.BoneWeights);
-	W.SetBoneWeights(InBones, InInfluences, NumEntries, InSettings);
-	return Result;
-}
-
-
-FBoneWeights FBoneWeights::Create(
-	TArrayView<const FBoneWeight> InBoneWeights, 
-	const FBoneWeightsSettings& InSettings /*= {} */
-)
-{
-	FBoneWeights Result;
-	FArrayWrapper(Result.BoneWeights).SetBoneWeights(InBoneWeights, InSettings);
-	return Result;
-}
-
-
-FBoneWeights FBoneWeights::Blend(
-    const FBoneWeights& InA,
-    const FBoneWeights& InB,
-    float InBias,
-    const FBoneWeightsSettings& InSettings /*= {} */
-)
-{
-	FBoneWeights Result;
-	FArrayWrapper W(Result.BoneWeights);
-	FArrayWrapper A(const_cast<FBoneWeights&>(InA).BoneWeights);
-	FArrayWrapper B(const_cast<FBoneWeights&>(InB).BoneWeights);
-	W.Blend(A, B, InBias, InSettings);
-	return Result;
-}
-
-
-void FBoneWeights::SortWeights()
-{
-	BoneWeights.Sort(WeightSortPredicate);
-}
-
-
-bool FBoneWeights::CullWeights(const FBoneWeightsSettings& InSettings)
-{
-	bool bCulled = false;
-	if (BoneWeights.Num() > InSettings.GetMaxWeightCount())
-	{
-		BoneWeights.SetNum(InSettings.GetMaxWeightCount(), false);
-		bCulled = true;
-	}
-
-	// If entries are now below the threshold, remove them.
-	while (BoneWeights.Num() > 0 && BoneWeights.Last().GetRawWeight() < InSettings.GetRawWeightThreshold())
-	{
-		BoneWeights.SetNum(BoneWeights.Num() - 1, false);
-		bCulled = true;
-	}
-
-	return bCulled;
-}
-
-
-void FBoneWeights::NormalizeWeights(
-    EBoneWeightNormalizeType InNormalizeType)
-{
-	// Early checks
-	if (InNormalizeType == EBoneWeightNormalizeType::None || BoneWeights.Num() == 0)
-	{
-		return;
-	}
-
-	// Common case.
-	if (BoneWeights.Num() == 1)
-	{
-		if (InNormalizeType == EBoneWeightNormalizeType::Always)
-		{
-			BoneWeights[0].SetRawWeight(FBoneWeight::GetMaxRawWeight());
-		}
-		return;
-	}
-
-	// We operate on int64, since we can easily end up with wraparound issues during one of the
-	// multiplications below when using int32. This would tank the division by WeightSum.
-	int64 WeightSum = 0;
-	for (const FBoneWeight& BW : BoneWeights)
-	{
-		WeightSum += BW.GetRawWeight();
-	}
-
-	if ((InNormalizeType == EBoneWeightNormalizeType::Always && ensure(WeightSum != 0)) ||
-	    WeightSum > FBoneWeight::GetMaxRawWeight())
-	{
-		int64 Correction = 0;
-
-		// Here we treat the raw weight as a 16.16 fixed point value and ensure that the
-		// fraction, which would otherwise be lost through rounding, is carried over to the 
-		// subsequent values to maintain a constant sum to the max weight value.
-		// We do this in descending weight order in an attempt to ensure that weight values 
-		// aren't needlessly lost after scaling.
-		for (FBoneWeight& BW : BoneWeights)
-		{
-			int64 ScaledWeight = int64(BW.GetRawWeight()) * FBoneWeight::GetMaxRawWeight() + Correction;
-			BW.SetRawWeight(uint16(FMath::Min(ScaledWeight / WeightSum, int64(FBoneWeight::GetMaxRawWeight()))));
-			Correction = ScaledWeight - BW.GetRawWeight() * WeightSum;
-		}
-	}
-}
-
-} // namespace AnimationCore
-} // namespace UE
diff --git a/Engine/Source/Runtime/AnimationCore/Public/BoneWeights.h b/Engine/Source/Runtime/AnimationCore/Public/BoneWeights.h
--- a/Engine/Source/Runtime/AnimationCore/Public/BoneWeights.h
+++ b/Engine/Source/Runtime/AnimationCore/Public/BoneWeights.h
@@ -2,6 +2,7 @@
 
 #pragma once
 
+#include "AnimationCore.h"
 #include "Algo/IsSorted.h"
 #include "BoneIndices.h"
 #include "Containers/UnrealString.h"
@@ -107,7 +108,7 @@ public:
 	void SetWeight(float InWeight)
 	{
 		InWeight = FMath::Clamp(InWeight, 0.0f, 1.0f);
-		RawWeight = uint16(InWeight * float(GetMaxRawWeight()) + 0.5f);
+		RawWeight = static_cast<uint16>(InWeight * static_cast<float>(GetMaxRawWeight()) + 0.5f);
 	}
 
 	/**
@@ -115,13 +116,13 @@ public:
 	*/
 	float GetWeight() const
 	{
-		return RawWeight / float(GetMaxRawWeight());
+		return RawWeight / static_cast<float>(GetMaxRawWeight());
 	}
 
 	/**
 	* Set the stored weight in the raw format. This avoids any floating point conversion.
 	*/
-	void SetRawWeight(uint16 InRawWeight)
+	void SetRawWeight(const uint16 InRawWeight)
 	{
 		RawWeight = InRawWeight;
 	}
@@ -144,7 +145,7 @@ public:
 	/** Returns a hash value from the bone weight values */
 	uint32 GetTypeHash() const
 	{
-		return HashCombine(::GetTypeHash(uint16(BoneIndex)), ::GetTypeHash(RawWeight));
+		return HashCombine(::GetTypeHash(static_cast<uint16>(BoneIndex)), ::GetTypeHash(RawWeight));
 	}
 
 	/**
@@ -155,6 +156,26 @@ public:
 		return FString::Printf(TEXT("<%d, %g>"), BoneIndex, GetWeight());
 	}
 
+
+	/**
+	* Returns the contents of this object as an int32.
+	*/
+	int32 ToInt32() const
+	{
+		return static_cast<int32>(static_cast<uint16>(BoneIndex) << 16 | RawWeight);
+	}
+
+	/**
+	* Converts an int32 to an FBoneWeight object and returns it.
+	*/
+	static FBoneWeight FromInt32(int32 InBoneWeight)
+	{
+		FBoneWeight BW;
+		BW.RawWeight = static_cast<uint16>(static_cast<uint32>(InBoneWeight) & 0xFFFFU);
+		BW.BoneIndex = static_cast<FBoneIndexType>(static_cast<uint32>(InBoneWeight) >> 16);
+		return BW;
+	}
+
 private:
 	FBoneIndexType BoneIndex;
 	uint16 RawWeight;
@@ -212,15 +233,15 @@ public:
 	FBoneWeightsSettings& SetWeightThreshold(float InWeightThreshold)
 	{
 		InWeightThreshold = FMath::Clamp(InWeightThreshold, 0.0f, 1.0f);
-		WeightThreshold = uint16(InWeightThreshold * FBoneWeight::GetMaxRawWeight() + 0.5f);
-		WeightThreshold = FMath::Max(WeightThreshold, uint16(1));
+		WeightThreshold = static_cast<uint16>(InWeightThreshold * FBoneWeight::GetMaxRawWeight() + 0.5f);
+		WeightThreshold = FMath::Max(WeightThreshold, static_cast<uint16>(1));
 		return *this;
 	}
 
 	/** Returns the weight threshold as a float value between (0, 1]. */
 	float GetWeightThreshold() const
 	{
-		return WeightThreshold / float(TNumericLimits<uint16>::Max());
+		return WeightThreshold / static_cast<float>(TNumericLimits<uint16>::Max());
 	}
 
 	/** Returns the raw weight threshold. This is the value used internally for weight
@@ -265,7 +286,7 @@ private:
 	EBoneWeightNormalizeType NormalizeType = EBoneWeightNormalizeType::Always;
 	int32 MaxWeightCount = MaxInlineBoneWeightCount;
 	uint16 WeightThreshold = 257; // == uint8(1) converted to 16 bit using 'v | v << 8'.
-	FBoneIndexType DefaultBoneIndex = FBoneIndexType(0);
+	FBoneIndexType DefaultBoneIndex = static_cast<FBoneIndexType>(0);
 	bool bHasDefaultBoneIndex = false;
 };
 
@@ -330,33 +351,51 @@ public:
 		Container(InContainer)
 	{ }
 
-	inline void SetBoneWeights(
+	template<typename OtherContainerAdapter, typename CT = ContainerType>
+	inline typename std::enable_if<!std::is_const<CT>::value, void>::type
+	SetBoneWeights(
+		const TBoneWeights<OtherContainerAdapter> &InBoneWeights,
+		const FBoneWeightsSettings& InSettings = {});
+
+	template<typename CT = ContainerType>
+    inline typename std::enable_if<!std::is_const<CT>::value, void>::type
+	SetBoneWeights(
 		TArrayView<const FBoneWeight> BoneWeights,
 		const FBoneWeightsSettings& InSettings = {});
 
-	inline void SetBoneWeights(
+	template<typename CT = ContainerType>
+	inline typename std::enable_if<!std::is_const<CT>::value, void>::type
+	SetBoneWeights(
 		const FBoneIndexType* InBones,
 		const float* InInfluences,
 		int32 NumEntries,
 		const FBoneWeightsSettings& InSettings = {}
 		);
 
-	inline void SetBoneWeights(
+	template<typename CT = ContainerType>
+	inline typename std::enable_if<!std::is_const<CT>::value, void>::type
+	SetBoneWeights(
 		const FBoneIndexType InBones[MaxInlineBoneWeightCount],
 		const uint8 InInfluences[MaxInlineBoneWeightCount],
 		const FBoneWeightsSettings& InSettings = {});
 
-	inline bool AddBoneWeight(
+	template<typename CT = ContainerType>
+    inline typename std::enable_if<!std::is_const<CT>::value, bool>::type
+	AddBoneWeight(
 		FBoneWeight InBoneWeight,
 		const FBoneWeightsSettings& InSettings = {}
 		);
 
-	inline bool RemoveBoneWeight(
+	template<typename CT = ContainerType>
+    inline typename std::enable_if<!std::is_const<CT>::value, bool>::type
+    RemoveBoneWeight(
 		FBoneIndexType InBoneIndex,
 		const FBoneWeightsSettings& InSettings = {}
 		);
 
-	inline void Renormalize(
+	template<typename CT = ContainerType>
+    inline typename std::enable_if<!std::is_const<CT>::value, void>::type
+    Renormalize(
 		const FBoneWeightsSettings& InSettings = {}
 		);
 
@@ -365,10 +404,11 @@ public:
 	  * that range may give unwanted results. 
 	  * NOTE: The current container can also be used as an input.
 	  */
-	template<typename ContainerTypeA, typename ContainerTypeB>
-	inline void Blend(
-		const TBoneWeights<ContainerTypeA>& InA,
-		const TBoneWeights<ContainerTypeB>& InB,
+	template<typename ContainerTypeA, typename ContainerTypeB, typename CT = ContainerType>
+    inline typename std::enable_if<!std::is_const<CT>::value, void>::type
+    Blend(
+		const TBoneWeights<ContainerTypeA>& InBoneWeightsA,
+		const TBoneWeights<ContainerTypeB>& InBoneWeightsB,
 		float InBias,
 		const FBoneWeightsSettings& InSettings = {});
 
@@ -421,7 +461,6 @@ public:
 		return Result;
 	}
 
-
 private:
 	inline void SetBoneWeightsInternal(
 		TArrayView<FBoneWeight> BoneWeights,
@@ -523,7 +562,7 @@ public:
       * for not passing the threshold value. If the weight was successfully incorporated, then
 	  * this function returns true. Otherwise it returns false.
 	  */
-	ANIMATIONCORE_API bool SetBoneWeight(
+	inline bool SetBoneWeight(
 	    FBoneWeight InBoneWeight,
 	    const FBoneWeightsSettings& InSettings = {});
 
@@ -538,17 +577,17 @@ public:
 	/** Removes a specific bone from the list of weights, re-normalizing and pruning bones,
 	  * if needed.
 	  */
-	ANIMATIONCORE_API bool RemoveBoneWeight(
+	inline bool RemoveBoneWeight(
 	    FBoneIndexType InBone,
 	    const FBoneWeightsSettings& InSettings = {});
 
 	/** Force normalization of weights. This is useful if a set of operations were performed
 	  * with no normalization, for efficiency, and normalization is needed post-operation.
 	  */
-	ANIMATIONCORE_API void Renormalize(const FBoneWeightsSettings& InSettings = {});
+	inline void Renormalize(const FBoneWeightsSettings& InSettings = {});
 
 	/** A helper to create a FBoneWeights container from FSoftSkinVertex data structure. */
-	ANIMATIONCORE_API static FBoneWeights Create(
+	static inline FBoneWeights Create(
 	    const FBoneIndexType InBones[MaxInlineBoneWeightCount],
 	    const uint8 InWeights[MaxInlineBoneWeightCount],
 	    const FBoneWeightsSettings& InSettings = {});
@@ -556,27 +595,39 @@ public:
 	/** A helper to create a FBoneWeights container from separated bone index and weight arrays.
 	  * The size of the two arrays *must* be the same -- otherwise the behavior is undefined.
 	  */
-	ANIMATIONCORE_API static FBoneWeights Create(
+	static inline FBoneWeights Create(
 	    const FBoneIndexType* InBones,
 	    const float* InWeights,
 	    int32 NumEntries,
 	    const FBoneWeightsSettings& InSettings = {});
 
 	/** A helper to create FBoneWeights container from a TArray of FBoneWeight objects. */
-	ANIMATIONCORE_API static FBoneWeights Create(
+	static inline FBoneWeights Create(
 	    TArrayView<const FBoneWeight> BoneWeights,
 	    const FBoneWeightsSettings& InSettings = {});
 
+	/** A helper to create FBoneWeights container from a TBoneWeights derivative object */
+	template<typename OtherContainerAdapter>
+	static FBoneWeights Create(
+		TBoneWeights<OtherContainerAdapter> InBoneWeights,
+		const FBoneWeightsSettings& InSettings = {})
+	{
+		FBoneWeights Result;
+		FArrayWrapper W(Result.BoneWeights);
+		W.SetBoneWeights(InBoneWeights, InSettings);
+		return Result;
+	}
+	
 	/** Blend two bone weights together, making sure to add in every influence from both,
 	  * using the given settings. The bias value should lie on the [0,1] interval. Values outside
 	  * that range may give unwanted results. 
 	  */
-	ANIMATIONCORE_API static FBoneWeights Blend(
+	static inline FBoneWeights Blend(
 	    const FBoneWeights& InA,
 	    const FBoneWeights& InB,
 	    float InBias,
 	    const FBoneWeightsSettings& InSettings = {});
-
+	
 	// Ranged-based for loop compatibility -- but only the const version.
 	using RangedForConstIteratorType = BoneWeightArrayT::RangedForConstIteratorType;
 
@@ -593,13 +644,18 @@ public:
 	  * @param Index The index of the desired bone weight value.
 	  * @return The bone weight value for that index, or undefined if outside the valid range.
 	  */
-	const FBoneWeight& operator[](int32 Index) { return BoneWeights.operator[](Index); }
-
+	const FBoneWeight& operator[](int32 Index) const { return BoneWeights.operator[](Index); }
 
+	/** Returns the FBoneWeight list as an array view */
+	TArrayView<const FBoneWeight> ToArrayView() const
+	{
+		return BoneWeights;
+	}
+	
 	// -- Helper functions
 
 	/** Find the bone weight corresponding to the given bone index. If a bone weight with this
-	  * index doesn't exist
+	  * bone index does not exist it, a value of INDEX_NONE is returned.
 	  */
 	int32 FindWeightIndexByBone(FBoneIndexType InBoneIndex) const
 	{
@@ -631,13 +687,9 @@ public:
 	}
 
 private:
-	static FBoneWeights CreateFromArrayView(
-	    TArrayView<FBoneWeight> BoneWeights,
-	    const FBoneWeightsSettings& InSettings = {});
-
-	void SortWeights();
-	bool CullWeights(const FBoneWeightsSettings& InSettings);
-	void NormalizeWeights(EBoneWeightNormalizeType InNormalizeType);
+	inline void SortWeights();
+	inline bool CullWeights(const FBoneWeightsSettings& InSettings);
+	inline void NormalizeWeights(EBoneWeightNormalizeType InNormalizeType);
 
 	friend uint32 GetTypeHash(const FBoneWeights& InBoneWeights);
 
@@ -648,14 +700,37 @@ private:
 
 /// TBoneWeights implementation
 template<typename ContainerAdapter>
-void TBoneWeights<ContainerAdapter>::SetBoneWeights(
-	TArrayView<const FBoneWeight> BoneWeights,
+template<typename OtherContainerAdapter, typename CT>
+typename std::enable_if<!std::is_const<CT>::value, void>::type
+TBoneWeights<ContainerAdapter>::SetBoneWeights(
+	const TBoneWeights<OtherContainerAdapter>& InBoneWeights,
 	const FBoneWeightsSettings& InSettings /*= {}*/)
 {
 	TArray<FBoneWeight, TInlineAllocator<MaxInlineBoneWeightCount>> StackBoneWeights;
+	StackBoneWeights.Reserve(InBoneWeights.Num());
+	for (int32 Index = 0; Index < InBoneWeights.Num(); Index++)
+	{
+		FBoneWeight BW = InBoneWeights.Get(Index);
+		if (BW.GetRawWeight() >= InSettings.GetRawWeightThreshold())
+		{
+			StackBoneWeights.Add(BW);
+		}
+	}	
 
-	StackBoneWeights.Reserve(BoneWeights.Num());
-	for (const FBoneWeight& BW : BoneWeights)
+	SetBoneWeightsInternal(StackBoneWeights, InSettings);
+}
+
+template<typename ContainerAdapter>
+template<typename CT>
+typename std::enable_if<!std::is_const<CT>::value, void>::type
+TBoneWeights<ContainerAdapter>::SetBoneWeights(
+	TArrayView<const FBoneWeight> InBoneWeights,
+	const FBoneWeightsSettings& InSettings /*= {}*/)
+{
+	TArray<FBoneWeight, TInlineAllocator<MaxInlineBoneWeightCount>> StackBoneWeights;
+
+	StackBoneWeights.Reserve(InBoneWeights.Num());
+	for (const FBoneWeight& BW : InBoneWeights)
 	{
 		if (BW.GetRawWeight() >= InSettings.GetRawWeightThreshold())
 		{
@@ -668,7 +743,9 @@ void TBoneWeights<ContainerAdapter>::SetBoneWeights(
 
 
 template<typename ContainerAdapter>
-void TBoneWeights<ContainerAdapter>::SetBoneWeights(
+template<typename CT>
+typename std::enable_if<!std::is_const<CT>::value, void>::type
+TBoneWeights<ContainerAdapter>::SetBoneWeights(
 	const FBoneIndexType* InBones,
 	const float* InInfluences,
 	int32 NumEntries,
@@ -691,7 +768,9 @@ void TBoneWeights<ContainerAdapter>::SetBoneWeights(
 
 
 template<typename ContainerAdapter>
-void TBoneWeights<ContainerAdapter>::SetBoneWeights(
+template<typename CT>
+typename std::enable_if<!std::is_const<CT>::value, void>::type
+TBoneWeights<ContainerAdapter>::SetBoneWeights(
 	const FBoneIndexType InBones[MaxInlineBoneWeightCount],
 	const uint8 InInfluences[MaxInlineBoneWeightCount],
 	const FBoneWeightsSettings& InSettings /*= {}*/)
@@ -757,7 +836,9 @@ void TBoneWeights<ContainerAdapter>::SetBoneWeightsInternal(
 }
 
 template<typename ContainerAdapter>
-bool TBoneWeights<ContainerAdapter>::AddBoneWeight(
+template<typename CT>
+typename std::enable_if<!std::is_const<CT>::value, bool>::type
+TBoneWeights<ContainerAdapter>::AddBoneWeight(
 	FBoneWeight InBoneWeight,
 	const FBoneWeightsSettings& InSettings /*= {}*/)
 {
@@ -766,7 +847,7 @@ bool TBoneWeights<ContainerAdapter>::AddBoneWeight(
 
 	// If the sum of weights could possibly exceed 1.0, we may need normalization based on
 	// the weight settings.
-	bool bMayNeedNormalization = false;
+	bool bMayNeedNormalization;
 
 	if (WeightIndex != INDEX_NONE)
 	{
@@ -806,7 +887,7 @@ bool TBoneWeights<ContainerAdapter>::AddBoneWeight(
 		}
 
 		// Are we already at the limit of weights for this container?
-		int32 NumWeights = ContainerAdapter::Num(Container);
+		const int32 NumWeights = ContainerAdapter::Num(Container);
 		if (NumWeights == InSettings.GetMaxWeightCount())
 		{
 			// If the weight is smaller than the smallest weight currently, then we reject.
@@ -841,7 +922,9 @@ bool TBoneWeights<ContainerAdapter>::AddBoneWeight(
 
 
 template<typename ContainerAdapter>
-bool TBoneWeights<ContainerAdapter>::RemoveBoneWeight(
+template<typename CT>
+typename std::enable_if<!std::is_const<CT>::value, bool>::type
+TBoneWeights<ContainerAdapter>::RemoveBoneWeight(
     FBoneIndexType InBoneIndex,
 	const FBoneWeightsSettings& InSettings /*= {}*/
 	)
@@ -869,7 +952,9 @@ bool TBoneWeights<ContainerAdapter>::RemoveBoneWeight(
 
 
 template<typename ContainerAdapter>
-void TBoneWeights<ContainerAdapter>::Renormalize(
+template<typename CT>
+typename std::enable_if<!std::is_const<CT>::value, void>::type
+TBoneWeights<ContainerAdapter>::Renormalize(
 	const FBoneWeightsSettings& InSettings /*= {}*/
 	)
 {
@@ -884,11 +969,12 @@ void TBoneWeights<ContainerAdapter>::Renormalize(
 
 
 template<typename ContainerAdapter>
-template<typename ContainerAdapterA, typename ContainerAdapterB>
-void TBoneWeights<ContainerAdapter>::Blend(
+template<typename ContainerAdapterA, typename ContainerAdapterB, typename CT>
+typename std::enable_if<!std::is_const<CT>::value, void>::type
+TBoneWeights<ContainerAdapter>::Blend(
 	const TBoneWeights<ContainerAdapterA>& InBoneWeightsA,
 	const TBoneWeights<ContainerAdapterB>& InBoneWeightsB,
-	float InBias,
+	const float InBias,
 	const FBoneWeightsSettings& InSettings /*= {}*/
 	)
 {
@@ -937,11 +1023,11 @@ void TBoneWeights<ContainerAdapter>::Blend(
 	TArray<FBoneWeight, TInlineAllocator<MaxInlineBoneWeightCount * 2>> BoneWeights;
 	BoneWeights.Reserve(InBoneWeightsA.Num() + InBoneWeightsB.Num());
 
-	int32 RawBiasB = int32(InBias * float(FBoneWeight::GetMaxRawWeight()));
-	int32 RawBiasA = FBoneWeight::GetMaxRawWeight() - RawBiasB;
+	const int32 RawBiasB = static_cast<int32>(InBias * static_cast<float>(FBoneWeight::GetMaxRawWeight()));
+	const int32 RawBiasA = FBoneWeight::GetMaxRawWeight() - RawBiasB;
 
 	int32 IndexA = 0, IndexB = 0;
-	for (; IndexA < InBoneWeightsA.Num() && IndexB < InBoneWeightsB.Num(); /**/)
+	for (; IndexA < InBoneWeightsA.Num() && IndexB < InBoneWeightsB.Num(); /* */ )
 	{
 		const FBoneWeight& BWA = InBoneWeightsA[IndirectIndexA[IndexA]];
 		const FBoneWeight& BWB = InBoneWeightsB[IndirectIndexB[IndexB]];
@@ -1023,7 +1109,7 @@ void TBoneWeights<ContainerAdapter>::NormalizeWeights(
 	EBoneWeightNormalizeType InNormalizeType
 	)
 {
-	int32 NumWeights = ContainerAdapter::Num(Container);
+	const int32 NumWeights = ContainerAdapter::Num(Container);
 
 	// Early checks
 	if (InNormalizeType == EBoneWeightNormalizeType::None || NumWeights == 0)
@@ -1065,8 +1151,8 @@ void TBoneWeights<ContainerAdapter>::NormalizeWeights(
 		for (int32 Index = 0; Index < NumWeights; Index++)
 		{
 			FBoneWeight BW = ContainerAdapter::Get(Container, Index);
-			int64 ScaledWeight = int64(BW.GetRawWeight()) * FBoneWeight::GetMaxRawWeight() + Correction;
-			BW.SetRawWeight(uint16(FMath::Min(ScaledWeight / WeightSum, int64(FBoneWeight::GetMaxRawWeight()))));
+			const int64 ScaledWeight = static_cast<int64>(BW.GetRawWeight()) * FBoneWeight::GetMaxRawWeight() + Correction;
+			BW.SetRawWeight(static_cast<uint16>(FMath::Min(ScaledWeight / WeightSum, static_cast<int64>(FBoneWeight::GetMaxRawWeight()))));
 			Correction = ScaledWeight - BW.GetRawWeight() * WeightSum;
 			ContainerAdapter::Set(Container, Index, BW);
 		}
@@ -1077,11 +1163,204 @@ void TBoneWeights<ContainerAdapter>::NormalizeWeights(
 template<typename ContainerAdapter>
 bool UE::AnimationCore::TBoneWeights<ContainerAdapter>::Verify() const
 {
-	// FIXME: TBD.
+	const int32 NumEntries = Num(); 
+	if (NumEntries == 0)
+	{
+		return true;
+	}
+	
+	// Check that bone indexes are unique
+	TSet<FBoneIndexType> BoneIndexes;
+	BoneIndexes.Reserve(NumEntries);
+	for (int32 Index = 0; Index < NumEntries; Index++)
+	{
+		FBoneWeight BW = Get(Index);
+		if (BoneIndexes.Find(BW.GetBoneIndex()) != nullptr)
+		{
+			// Commented out for now, to avoid linker errors, since this is a header-only file.
+			// UE_LOG(LogAnimationCore, Error, TEXT("Bone Index %d is duplicated"), static_cast<int>(BW.GetBoneIndex()));
+			return false;
+		}
+		BoneIndexes.Add(BW.GetBoneIndex());
+	}
+	
+	// Check that all weights are ordered.
+	float LastWeight = Get(0).GetWeight();
+	for (int32 Index = 1; Index < NumEntries; Index++)
+	{
+		if (Get(Index).GetWeight() > LastWeight)
+		{
+			//UE_LOG(LogAnimationCore, Error, TEXT("Bone Weight at %d is greater than previous (%g > %g)"),
+			//	Index, Get(Index).GetWeight(), LastWeight);
+			return false;
+		}
+		LastWeight = Get(Index).GetWeight();
+	}		
+	
 	return true;
 }
 
 
+// FBoneWeights implementations
+static auto WeightSortPredicate = [](const FBoneWeight& A, const FBoneWeight& B) {
+	return A.GetRawWeight() > B.GetRawWeight();
+};
+
+
+bool FBoneWeights::SetBoneWeight(
+    FBoneWeight InBoneWeight,
+    const FBoneWeightsSettings& InSettings /*= {} */
+)
+{
+	FArrayWrapper W(BoneWeights);
+	return W.AddBoneWeight(InBoneWeight, InSettings);
+}
+
+
+bool FBoneWeights::RemoveBoneWeight(
+    FBoneIndexType InBoneIndex,
+    const FBoneWeightsSettings& InSettings /*= {} */
+)
+{
+	FArrayWrapper W(BoneWeights);
+	return W.RemoveBoneWeight(InBoneIndex, InSettings);
+}
+
+
+void FBoneWeights::Renormalize(const FBoneWeightsSettings& InSettings /*= {}*/)
+{
+	FArrayWrapper W(BoneWeights);
+	return W.Renormalize(InSettings);
+}
+
+
+FBoneWeights FBoneWeights::Create(
+    const FBoneIndexType InBones[MaxInlineBoneWeightCount],
+    const uint8 InInfluences[MaxInlineBoneWeightCount],
+    const FBoneWeightsSettings& InSettings /*= {} */
+)
+{
+	FBoneWeights Result;
+	FArrayWrapper W(Result.BoneWeights);
+	W.SetBoneWeights(InBones, InInfluences, InSettings);
+	return Result;
+}
+
+
+FBoneWeights FBoneWeights::Create(
+    const FBoneIndexType* InBones,
+    const float* InInfluences,
+    int32 NumEntries,
+    const FBoneWeightsSettings& InSettings /*= {} */
+)
+{
+	FBoneWeights Result;
+	FArrayWrapper W(Result.BoneWeights);
+	W.SetBoneWeights(InBones, InInfluences, NumEntries, InSettings);
+	return Result;
+}
+
+
+FBoneWeights FBoneWeights::Create(
+	TArrayView<const FBoneWeight> InBoneWeights, 
+	const FBoneWeightsSettings& InSettings /*= {} */
+)
+{
+	FBoneWeights Result;
+	FArrayWrapper(Result.BoneWeights).SetBoneWeights(InBoneWeights, InSettings);
+	return Result;
+}
+
+
+FBoneWeights FBoneWeights::Blend(
+    const FBoneWeights& InA,
+    const FBoneWeights& InB,
+    float InBias,
+    const FBoneWeightsSettings& InSettings /*= {} */
+)
+{
+	FBoneWeights Result;
+	FArrayWrapper W(Result.BoneWeights);
+	FArrayWrapper A(const_cast<FBoneWeights&>(InA).BoneWeights);
+	FArrayWrapper B(const_cast<FBoneWeights&>(InB).BoneWeights);
+	W.Blend(A, B, InBias, InSettings);
+	return Result;
+}
+
+
+void FBoneWeights::SortWeights()
+{
+	BoneWeights.Sort(WeightSortPredicate);
+}
+
+
+bool FBoneWeights::CullWeights(const FBoneWeightsSettings& InSettings)
+{
+	bool bCulled = false;
+	if (BoneWeights.Num() > InSettings.GetMaxWeightCount())
+	{
+		BoneWeights.SetNum(InSettings.GetMaxWeightCount(), false);
+		bCulled = true;
+	}
+
+	// If entries are now below the threshold, remove them.
+	while (BoneWeights.Num() > 0 && BoneWeights.Last().GetRawWeight() < InSettings.GetRawWeightThreshold())
+	{
+		BoneWeights.SetNum(BoneWeights.Num() - 1, false);
+		bCulled = true;
+	}
+
+	return bCulled;
+}
+
+
+void FBoneWeights::NormalizeWeights(
+    EBoneWeightNormalizeType InNormalizeType)
+{
+	// Early checks
+	if (InNormalizeType == EBoneWeightNormalizeType::None || BoneWeights.Num() == 0)
+	{
+		return;
+	}
+
+	// Common case.
+	if (BoneWeights.Num() == 1)
+	{
+		if (InNormalizeType == EBoneWeightNormalizeType::Always)
+		{
+			BoneWeights[0].SetRawWeight(FBoneWeight::GetMaxRawWeight());
+		}
+		return;
+	}
+
+	// We operate on int64, since we can easily end up with wraparound issues during one of the
+	// multiplications below when using int32. This would tank the division by WeightSum.
+	int64 WeightSum = 0;
+	for (const FBoneWeight& BW : BoneWeights)
+	{
+		WeightSum += BW.GetRawWeight();
+	}
+
+	if ((InNormalizeType == EBoneWeightNormalizeType::Always && ensure(WeightSum != 0)) ||
+	    WeightSum > FBoneWeight::GetMaxRawWeight())
+	{
+		int64 Correction = 0;
+
+		// Here we treat the raw weight as a 16.16 fixed point value and ensure that the
+		// fraction, which would otherwise be lost through rounding, is carried over to the 
+		// subsequent values to maintain a constant sum to the max weight value.
+		// We do this in descending weight order in an attempt to ensure that weight values 
+		// aren't needlessly lost after scaling.
+		for (FBoneWeight& BW : BoneWeights)
+		{
+			int64 ScaledWeight = int64(BW.GetRawWeight()) * FBoneWeight::GetMaxRawWeight() + Correction;
+			BW.SetRawWeight(uint16(FMath::Min(ScaledWeight / WeightSum, int64(FBoneWeight::GetMaxRawWeight()))));
+			Correction = ScaledWeight - BW.GetRawWeight() * WeightSum;
+		}
+	}
+}
+	
+
 } // namespace AnimationCore
 } // namespace UE
 
diff --git a/Engine/Source/Runtime/SkeletalMeshDescription/Private/SkeletalMeshAttributes.cpp b/Engine/Source/Runtime/SkeletalMeshDescription/Private/SkeletalMeshAttributes.cpp
--- a/Engine/Source/Runtime/SkeletalMeshDescription/Private/SkeletalMeshAttributes.cpp
+++ b/Engine/Source/Runtime/SkeletalMeshDescription/Private/SkeletalMeshAttributes.cpp
@@ -4,19 +4,96 @@
 
 namespace MeshAttribute
 {
-	const FName Vertex::InfluenceCount("InfluenceCount");
-	const FName Vertex::InfluenceBone("InfluenceBone");
-	const FName Vertex::InfluenceWeight("InfluenceWeight");
+	const FName Vertex::SkinWeights("SkinWeights");
 }
 
 
+FName FSkeletalMeshAttributes::DefaultSkinWeightProfileName("Default");
+
+
+static FString SkinWeightAttributeNamePrefix()
+{
+	return MeshAttribute::Vertex::SkinWeights.ToString() + TEXT("-");
+}
+
 
 void FSkeletalMeshAttributes::Register()
 {
-	// Add skeletalmesh attribute
-	MeshDescription.VertexAttributes().RegisterAttribute<int32[]>(MeshAttribute::Vertex::InfluenceBone, 1, 0, EMeshAttributeFlags::IndexReference | EMeshAttributeFlags::Mandatory);
-	MeshDescription.VertexAttributes().RegisterAttribute<float[]>(MeshAttribute::Vertex::InfluenceWeight, 1, 0.0f, EMeshAttributeFlags::Mandatory);
+	MeshDescription.VertexAttributes().RegisterAttribute<int32[]>(MeshAttribute::Vertex::SkinWeights, 1, 0.0f, EMeshAttributeFlags::Mandatory);
 
 	// Call super class
 	FStaticMeshAttributes::Register();
 }
+
+bool FSkeletalMeshAttributes::RegisterSkinWeightAttribute(const FName InProfileName)
+{
+	if (!IsValidSkinWeightProfileName(InProfileName))
+	{
+		return false;
+	}	
+	
+	const FName AttributeName = CreateSkinWeightAttributeName(InProfileName);
+	if (ensure(AttributeName.IsValid()))
+	{
+		TArray<FName> AllAttributeNames;
+		MeshDescription.VertexAttributes().GetAttributeNames(AllAttributeNames);
+		if (AllAttributeNames.Contains(AttributeName))
+		{
+			return false;
+		}
+
+		return MeshDescription.VertexAttributes().RegisterAttribute<int32[]>(AttributeName, 1, 0, EMeshAttributeFlags::Mandatory).IsValid();
+	}
+	else
+	{
+		return false;
+	}	
+}
+
+TArray<FName> FSkeletalMeshAttributes::GetSkinWeightProfileNames() const
+{
+	TArray<FName> AllAttributeNames;
+	MeshDescription.VertexAttributes().GetAttributeNames(AllAttributeNames);
+
+	TArray<FName> SkinWeightProfileNames;
+	const FString Prefix = SkinWeightAttributeNamePrefix();
+	bool bHasDefault = false; 
+
+	for (const FName AttributeName: AllAttributeNames)
+	{
+		if (AttributeName == MeshAttribute::Vertex::SkinWeights)
+		{
+			bHasDefault = true;
+		}
+		else if (AttributeName.ToString().StartsWith(Prefix))
+		{
+			FName ProfileName = FName(AttributeName.ToString().Mid(Prefix.Len()));
+
+			SkinWeightProfileNames.Add(ProfileName);
+		}
+	}
+
+	SkinWeightProfileNames.Sort([](const FName A, const FName B) -> bool { return A.FastLess(B); });
+	if (bHasDefault)
+	{
+		SkinWeightProfileNames.Insert(DefaultSkinWeightProfileName, 0);
+	}
+	
+	return SkinWeightProfileNames;
+}
+
+bool FSkeletalMeshAttributes::IsValidSkinWeightProfileName(const FName InProfileName)
+{
+	return !InProfileName.IsNone() && !InProfileName.IsEqual(DefaultSkinWeightProfileName, ENameCase::IgnoreCase);
+}
+
+FName FSkeletalMeshAttributes::CreateSkinWeightAttributeName(const FName InProfileName)
+{
+	// If it's the default profile, then return the base skin weights attribute name.
+	if (InProfileName.IsNone() || InProfileName.IsEqual(DefaultSkinWeightProfileName, ENameCase::IgnoreCase))
+	{
+		return MeshAttribute::Vertex::SkinWeights;
+	}	
+
+	return FName(SkinWeightAttributeNamePrefix() + InProfileName.ToString());
+}
diff --git a/Engine/Source/Runtime/SkeletalMeshDescription/Public/SkeletalMeshAttributes.h b/Engine/Source/Runtime/SkeletalMeshDescription/Public/SkeletalMeshAttributes.h
--- a/Engine/Source/Runtime/SkeletalMeshDescription/Public/SkeletalMeshAttributes.h
+++ b/Engine/Source/Runtime/SkeletalMeshDescription/Public/SkeletalMeshAttributes.h
@@ -4,6 +4,7 @@
 
 #include "MeshAttributes.h"
 #include "StaticMeshAttributes.h"
+#include "SkinWeightsAttributesRef.h"
 
 //Add any skeletalmesh specific attributes here
 
@@ -11,13 +12,12 @@ namespace MeshAttribute
 {
  	namespace Vertex
  	{
-		extern SKELETALMESHDESCRIPTION_API const FName InfluenceCount;
- 		extern SKELETALMESHDESCRIPTION_API const FName InfluenceBone;
-		extern SKELETALMESHDESCRIPTION_API const FName InfluenceWeight;
+		extern SKELETALMESHDESCRIPTION_API const FName SkinWeights;
  	}
 }
 
 
+
 class SKELETALMESHDESCRIPTION_API FSkeletalMeshAttributes : public FStaticMeshAttributes
 {
 public:
@@ -26,16 +26,40 @@ public:
 		: FStaticMeshAttributes(InMeshDescription)
 	{}
 
-	//TODO support uint16 in mesh description
-	TVertexAttributesRef<TArrayAttribute<int32>> GetVertexInfluenceBones() { return MeshDescription.VertexAttributes().GetAttributesRef<TArrayAttribute<int32>>(MeshAttribute::Vertex::InfluenceBone); }
-	TVertexAttributesConstRef<TArrayAttribute<int32>> GetVertexInfluenceBones() const { return MeshDescription.VertexAttributes().GetAttributesRef<TArrayAttribute<int32>>(MeshAttribute::Vertex::InfluenceBone); }
-
-	TVertexAttributesRef<TArrayAttribute<float>> GetVertexInfluenceWeights() { return MeshDescription.VertexAttributes().GetAttributesRef<TArrayAttribute<float>>(MeshAttribute::Vertex::InfluenceWeight); }
-	TVertexAttributesConstRef<TArrayAttribute<float>> GetVertexInfluenceWeights() const { return MeshDescription.VertexAttributes().GetAttributesRef<TArrayAttribute<float>>(MeshAttribute::Vertex::InfluenceWeight); }
-
+	// The name of the default skin weight profile.
+	static FName DefaultSkinWeightProfileName;
+	
 	virtual void Register() override;
 
-
+	/// Register a new skin weight profile with the given name. The attribute name will encode the profile name and
+	/// it will be listed in GetSkinWeightProfileNames(). Returns \c true if the profile was successfully registered.
+	/// Returns \c false if the attribute was already registered or if IsValidSkinWeightProfileName() returned false.
+	bool RegisterSkinWeightAttribute(const FName InProfileName);
+
+	/// Returns the list of all registered skin weight profile names on this mesh.
+	TArray<FName> GetSkinWeightProfileNames() const;
+	
+	/// Returns \c true if the given identifier is a valid profile name. If the name is empty, or matches the default profile,
+	/// then the profile name is considered invalid. 
+	static bool IsValidSkinWeightProfileName(const FName InProfileName);
+
+	/// Returns the skin weight profile given by its name. NAME_None corresponds to the default profile.
+	FSkinWeightsVertexAttributesRef GetVertexSkinWeights(const FName InProfileName = NAME_None)
+	{
+		return MeshDescription.VertexAttributes().GetAttributesRef<TArrayAttribute<int32>>(CreateSkinWeightAttributeName(InProfileName));
+	}
+	FSkinWeightsVertexAttributesConstRef GetVertexSkinWeights(const FName InProfileName = NAME_None) const
+	{
+		return MeshDescription.VertexAttributes().GetAttributesRef<TArrayAttribute<int32>>(CreateSkinWeightAttributeName(InProfileName));
+	}
+
+protected:
+	friend class FSkeletalMeshConstAttributes;
+
+	/// Construct a name for a skin weight attribute with the given skin weight profile name.
+	/// Each mesh description can hold different skin weight profiles, although the default
+	/// is always present.
+	static FName CreateSkinWeightAttributeName(const FName InProfileName);
 };
 
 
@@ -47,6 +71,8 @@ public:
 		: FStaticMeshConstAttributes(InMeshDescription)
 	{}
 
-	TVertexAttributesConstRef<TArrayAttribute<int32>> GetVertexInfluenceBones() const { return MeshDescription.VertexAttributes().GetAttributesRef<TArrayAttribute<int32>>(MeshAttribute::Vertex::InfluenceBone); }
-	TVertexAttributesConstRef<TArrayAttribute<float>> GetVertexInfluenceWeights() const { return MeshDescription.VertexAttributes().GetAttributesRef<TArrayAttribute<float>>(MeshAttribute::Vertex::InfluenceWeight); }
+	FSkinWeightsVertexAttributesConstRef GetVertexSkinWeights(const FName InProfileName = NAME_None) const
+	{
+		return MeshDescription.VertexAttributes().GetAttributesRef<TArrayAttribute<int32>>(FSkeletalMeshAttributes::CreateSkinWeightAttributeName(InProfileName));
+	}
 };
diff --git a/Engine/Source/Runtime/SkeletalMeshDescription/Public/SkinWeightsAttributesRef.h b/Engine/Source/Runtime/SkeletalMeshDescription/Public/SkinWeightsAttributesRef.h
new file mode 100644
--- /dev/null
+++ b/Engine/Source/Runtime/SkeletalMeshDescription/Public/SkinWeightsAttributesRef.h
@@ -0,0 +1,284 @@
+// Copyright Epic Games, Inc. All Rights Reserved.
+
+#pragma once
+
+#include "MeshDescription.h"
+#include "BoneWeights.h"
+
+
+namespace Internal
+{
+	struct FAttributesRefAdapterConst
+	{
+		using FBoneWeight = UE::AnimationCore::FBoneWeight;
+
+		using ContainerType = const TArrayAttribute<const int32>;
+
+		static int32 Num(const ContainerType& InContainer)
+		{
+			return InContainer.Num();
+		}
+
+		static FBoneWeight Get(const ContainerType& InContainer, int32 InIndex)
+		{
+			return FBoneWeight::FromInt32(InContainer.GetData()[InIndex]);
+		}
+
+		template<typename Predicate>
+        static int32 IndexOf(const ContainerType& InContainer, Predicate InPredicate)
+		{
+			auto ProxyPredicate = [InPredicate](const int32 InV) -> bool
+			{
+				return InPredicate(FBoneWeight::FromInt32(InV));
+			};
+			return InContainer.IndexOfByPredicate(ProxyPredicate);
+		}
+	};
+
+	struct FAttributesRefAdapter
+	{
+		using FBoneWeight = UE::AnimationCore::FBoneWeight;
+
+		using ContainerType = TArrayAttribute<int32>;
+
+		static void SetNum(ContainerType& InContainer, int32 InNum)
+		{
+			InContainer.SetNumUninitialized(InNum);
+		}
+
+		static int32 Num(const ContainerType& InContainer)
+		{
+			return InContainer.Num();
+		}
+
+		static FBoneWeight Get(const ContainerType& InContainer, int32 InIndex)
+		{
+			return FBoneWeight::FromInt32(InContainer.GetData()[InIndex]);
+		}
+
+		static void Set(ContainerType& InContainer, int32 InIndex, FBoneWeight InBoneWeight)
+		{
+			InContainer.GetData()[InIndex] = InBoneWeight.ToInt32();
+		}
+
+		static void Add(ContainerType& InContainer, FBoneWeight InBoneWeight)
+		{
+			InContainer.Add(InBoneWeight.ToInt32());
+		}
+
+		static void Remove(ContainerType& InContainer, int32 InIndex)
+		{
+			InContainer.RemoveAt(InIndex, 1);
+		}
+
+		template<typename Predicate>
+        static void Sort(ContainerType& InContainer, Predicate InPredicate)
+		{
+			auto ProxyPredicate = [InPredicate](const int32 InA, const int32 InB) -> bool
+			{
+				return InPredicate(FBoneWeight::FromInt32(InA), FBoneWeight::FromInt32(InB));
+			};
+			InContainer.SortByPredicate(ProxyPredicate);
+		}
+
+		template<typename Predicate>
+        static int32 IndexOf(const ContainerType& InContainer, Predicate InPredicate)
+		{
+			auto ProxyPredicate = [InPredicate](const int32 InV) -> bool
+			{
+				return InPredicate(FBoneWeight::FromInt32(InV));
+			};
+			return InContainer.IndexOfByPredicate(ProxyPredicate);
+		}
+	};
+
+
+	template<typename AdaptorType, typename StorageType>
+	class TVertexBoneWeightsBase :
+	    public UE::AnimationCore::TBoneWeights<AdaptorType>
+	{
+		using Super = UE::AnimationCore::TBoneWeights<AdaptorType>;
+	public:	
+		explicit TVertexBoneWeightsBase(const TArrayAttribute<StorageType> InContainer) :
+	        Super(Container),
+	        Container(InContainer)
+		{}
+
+		// For compatibility with range-based for loops.
+		class const_iterator
+		{
+		public:
+			const_iterator() = default;
+
+			const_iterator& operator++()
+			{
+				if (Index != INDEX_NONE && Weights)
+				{
+					++Index;
+					if (Index == Weights->Num())
+					{
+						Weights = nullptr;
+						Index = INDEX_NONE;
+					}
+				}
+				return *this;
+			}
+
+			bool operator!=(const const_iterator &InOther) const
+			{
+				return Weights != InOther.Weights || Index != InOther.Index;
+			}			
+
+			UE::AnimationCore::FBoneWeight operator*() const
+			{
+				return Weights->Get(Index);
+			}
+		
+		protected:
+			friend class TVertexBoneWeightsBase<AdaptorType, StorageType>;
+
+			explicit const_iterator(const TVertexBoneWeightsBase<AdaptorType, StorageType> *InWeights) :
+	            Weights(InWeights),
+	            Index(0)
+			{
+			}
+
+			const TVertexBoneWeightsBase<AdaptorType, StorageType> *Weights = nullptr;
+			int32 Index = INDEX_NONE;
+		};
+
+		const_iterator begin() const { return const_iterator(this); }		
+		const_iterator end() const { return const_iterator(); }
+
+	private:
+		TArrayAttribute<StorageType> Container;
+	};
+
+} // namespace Internal
+
+using FVertexBoneWeights = Internal::TVertexBoneWeightsBase<Internal::FAttributesRefAdapter, int32>;
+using FVertexBoneWeightsConst = Internal::TVertexBoneWeightsBase<Internal::FAttributesRefAdapterConst, const int32>;
+
+
+// Merge these two.
+class FSkinWeightsVertexAttributesRef
+{
+public:
+	FSkinWeightsVertexAttributesRef() = default;
+	FSkinWeightsVertexAttributesRef(const FSkinWeightsVertexAttributesRef &) = default;
+	FSkinWeightsVertexAttributesRef(FSkinWeightsVertexAttributesRef &&) noexcept = default;
+	FSkinWeightsVertexAttributesRef& operator=(const FSkinWeightsVertexAttributesRef &InAttributesRef)
+	{
+		AttributesRef = InAttributesRef.AttributesRef;
+		return *this;
+	}
+	
+	FSkinWeightsVertexAttributesRef& operator=(FSkinWeightsVertexAttributesRef &&InAttributesRef) noexcept
+	{
+		AttributesRef = MoveTemp(InAttributesRef.AttributesRef);
+		InAttributesRef.AttributesRef = TVertexAttributesRef<TArrayAttribute<int32>>();
+		return *this;
+	}
+
+	bool IsValid() const
+	{
+		return AttributesRef.IsValid();
+	}
+
+	FVertexBoneWeights Get(const FVertexID InVertexIndex) const
+	{
+		return FVertexBoneWeights(AttributesRef.Get(InVertexIndex));
+	}
+
+	void Set(const FVertexID InVertexIndex, const UE::AnimationCore::FBoneWeights &InWeights)
+	{
+		// FBoneWeights guarantees ordering and normalization, so we just copy the raw data.
+		TArrayAttribute<int32> Weights = AttributesRef.Get(InVertexIndex);
+
+		Weights.SetNumUninitialized(InWeights.Num());
+		for (int32 Index = 0; Index < InWeights.Num(); Index++)
+		{
+			Weights[Index] = InWeights[Index].ToInt32();
+		}
+	}
+
+	void Set(
+		const FVertexID InVertexIndex, 
+		TArrayView<const UE::AnimationCore::FBoneWeight> InWeights,
+		const UE::AnimationCore::FBoneWeightsSettings& InSettings = {}
+		)
+	{
+		// Proxy through FVertexBoneWeights so we get proper sorting and normalization of the weights. 
+		FVertexBoneWeights BoneWeights(AttributesRef.Get(InVertexIndex));
+
+		BoneWeights.SetBoneWeights(InWeights, InSettings);
+	}
+	
+protected:
+	friend class FSkeletalMeshAttributes;
+
+	FSkinWeightsVertexAttributesRef(TVertexAttributesRef<TArrayAttribute<int32>> InAttributesRef)
+        : AttributesRef(InAttributesRef)
+	{}
+
+private:
+	friend class FSkinWeightsVertexAttributesConstRef;
+	TVertexAttributesRef<TArrayAttribute<int32>> AttributesRef;
+};
+
+
+class FSkinWeightsVertexAttributesConstRef
+{
+public:
+	FSkinWeightsVertexAttributesConstRef() = default;
+	FSkinWeightsVertexAttributesConstRef(const FSkinWeightsVertexAttributesConstRef &) = default;
+	FSkinWeightsVertexAttributesConstRef(FSkinWeightsVertexAttributesConstRef &&) noexcept = default;
+	FSkinWeightsVertexAttributesConstRef& operator=(const FSkinWeightsVertexAttributesConstRef &InAttributesRef)
+	{
+		AttributesConstRef = InAttributesRef.AttributesConstRef;
+		return *this;
+	}
+	
+	FSkinWeightsVertexAttributesConstRef& operator=(FSkinWeightsVertexAttributesConstRef &&InAttributesRef) noexcept
+	{
+		AttributesConstRef = MoveTemp(InAttributesRef.AttributesConstRef);
+		InAttributesRef.AttributesConstRef = TVertexAttributesConstRef<TArrayAttribute<int32>>();
+		return *this;
+	}
+
+	// Converting constructors from the non-const variant
+	explicit FSkinWeightsVertexAttributesConstRef(const FSkinWeightsVertexAttributesRef &InAttribs) :
+		AttributesConstRef(InAttribs.AttributesRef)
+	{}
+	
+	FSkinWeightsVertexAttributesConstRef& operator=(const FSkinWeightsVertexAttributesRef &InAttribs)
+	{
+		AttributesConstRef = InAttribs.AttributesRef;
+		return *this;
+	}
+	
+	bool IsValid() const
+	{
+		return AttributesConstRef.IsValid();
+	}
+	
+	FVertexBoneWeightsConst Get(FVertexID InVertexIndex) const
+	{
+		return FVertexBoneWeightsConst(AttributesConstRef.Get(InVertexIndex));
+	}
+
+protected:
+	friend class FSkeletalMeshAttributes;
+	friend class FSkeletalMeshConstAttributes;
+
+	FSkinWeightsVertexAttributesConstRef(TVertexAttributesRef<TArrayAttribute<int32>> InAttributesConstRef)
+        : AttributesConstRef(InAttributesConstRef)
+	{}
+
+	FSkinWeightsVertexAttributesConstRef(TVertexAttributesRef<TArrayAttribute<const int32>> InAttributesConstRef)
+	    : AttributesConstRef(InAttributesConstRef)
+	{}
+
+private:
+	TVertexAttributesConstRef<TArrayAttribute<int32>> AttributesConstRef;
+};
diff --git a/Engine/Source/Runtime/SkeletalMeshDescription/SkeletalMeshDescription.Build.cs b/Engine/Source/Runtime/SkeletalMeshDescription/SkeletalMeshDescription.Build.cs
--- a/Engine/Source/Runtime/SkeletalMeshDescription/SkeletalMeshDescription.Build.cs
+++ b/Engine/Source/Runtime/SkeletalMeshDescription/SkeletalMeshDescription.Build.cs
@@ -9,6 +9,17 @@ namespace UnrealBuildTool.Rules
             PrivateIncludePaths.Add("Runtime/SkeletalMeshDescription/Private");
             PublicIncludePaths.Add("Runtime/SkeletalMeshDescription/Public");
 
+            // For GPUSkinPublicDefs.h
+            PublicIncludePaths.Add("Runtime/Engine/Public");
+            
+            // For BoneWeights
+			PublicIncludePathModuleNames.AddRange(
+				new string[]
+				{
+					"AnimationCore"
+				}
+			);
+
 			PublicDependencyModuleNames.AddRange(
 				new string[]
 				{
