Fixed an edge case where if we're normalizing weights which are all equal, we end up in a situation where the error distribution makes the last entry have a slightly greater weight value than the previous, which invalidates the descending weight order invariant. If we detect this case, just reorder after normalization.
Fixed the Verify function so that we can see the current weight in the debugger.
Removed dead functions.

#trivial
#rnx

diff --git a/Engine/Source/Runtime/AnimationCore/Public/BoneWeights.h b/Engine/Source/Runtime/AnimationCore/Public/BoneWeights.h
--- a/Engine/Source/Runtime/AnimationCore/Public/BoneWeights.h
+++ b/Engine/Source/Runtime/AnimationCore/Public/BoneWeights.h
@@ -11,6 +11,8 @@
 #include "Serialization/Archive.h"
 #include "Templates/TypeHash.h"
 
+#include <limits>
+
 namespace UE { 
 namespace AnimationCore {
 
@@ -687,10 +689,6 @@ public:
 	}
 
 private:
-	inline void SortWeights();
-	inline bool CullWeights(const FBoneWeightsSettings& InSettings);
-	inline void NormalizeWeights(EBoneWeightNormalizeType InNormalizeType);
-
 	friend uint32 GetTypeHash(const FBoneWeights& InBoneWeights);
 
 	/// List of bone weights, in order of descending weight.
@@ -1148,13 +1146,29 @@ void TBoneWeights<ContainerAdapter>::NormalizeWeights(
 		// subsequent values to maintain a constant sum to the max weight value.
 		// We do this in descending weight order in an attempt to ensure that weight values
 		// aren't needlessly lost after scaling.
+		// It can happen, if all the weights are the same, or similar, that the last weight may
+		// now be greater in value than the others. In that case we re-sort to ensure that our
+		// descending weight order invariant holds. 
+		bool bNeedResort = false;
+		uint16 LastWeight = std::numeric_limits<uint16>::max();
 		for (int32 Index = 0; Index < NumWeights; Index++)
 		{
 			FBoneWeight BW = ContainerAdapter::Get(Container, Index);
 			const int64 ScaledWeight = static_cast<int64>(BW.GetRawWeight()) * FBoneWeight::GetMaxRawWeight() + Correction;
-			BW.SetRawWeight(static_cast<uint16>(FMath::Min(ScaledWeight / WeightSum, static_cast<int64>(FBoneWeight::GetMaxRawWeight()))));
+			const uint16 NewWeight = static_cast<uint16>(FMath::Min(ScaledWeight / WeightSum, static_cast<int64>(FBoneWeight::GetMaxRawWeight()))); 
+			BW.SetRawWeight(NewWeight);
 			Correction = ScaledWeight - BW.GetRawWeight() * WeightSum;
 			ContainerAdapter::Set(Container, Index, BW);
+			if (NewWeight > LastWeight)
+			{
+				bNeedResort = true;
+			}
+			LastWeight = NewWeight;
+		}
+
+		if (bNeedResort)
+		{
+			SortWeights();
 		}
 	}
 }
@@ -1188,13 +1202,14 @@ bool UE::AnimationCore::TBoneWeights<ContainerAdapter>::Verify() const
 	float LastWeight = Get(0).GetWeight();
 	for (int32 Index = 1; Index < NumEntries; Index++)
 	{
-		if (Get(Index).GetWeight() > LastWeight)
+		const float Weight = Get(Index).GetWeight();
+		if (Weight > LastWeight)
 		{
 			//UE_LOG(LogAnimationCore, Error, TEXT("Bone Weight at %d is greater than previous (%g > %g)"),
-			//	Index, Get(Index).GetWeight(), LastWeight);
+			//	Index, Weight, LastWeight);
 			return false;
 		}
-		LastWeight = Get(Index).GetWeight();
+		LastWeight = Weight;
 	}		
 	
 	return true;
@@ -1288,79 +1303,6 @@ FBoneWeights FBoneWeights::Blend(
 }
 
 
-void FBoneWeights::SortWeights()
-{
-	BoneWeights.Sort(WeightSortPredicate);
-}
-
-
-bool FBoneWeights::CullWeights(const FBoneWeightsSettings& InSettings)
-{
-	bool bCulled = false;
-	if (BoneWeights.Num() > InSettings.GetMaxWeightCount())
-	{
-		BoneWeights.SetNum(InSettings.GetMaxWeightCount(), false);
-		bCulled = true;
-	}
-
-	// If entries are now below the threshold, remove them.
-	while (BoneWeights.Num() > 0 && BoneWeights.Last().GetRawWeight() < InSettings.GetRawWeightThreshold())
-	{
-		BoneWeights.SetNum(BoneWeights.Num() - 1, false);
-		bCulled = true;
-	}
-
-	return bCulled;
-}
-
-
-void FBoneWeights::NormalizeWeights(
-    EBoneWeightNormalizeType InNormalizeType)
-{
-	// Early checks
-	if (InNormalizeType == EBoneWeightNormalizeType::None || BoneWeights.Num() == 0)
-	{
-		return;
-	}
-
-	// Common case.
-	if (BoneWeights.Num() == 1)
-	{
-		if (InNormalizeType == EBoneWeightNormalizeType::Always)
-		{
-			BoneWeights[0].SetRawWeight(FBoneWeight::GetMaxRawWeight());
-		}
-		return;
-	}
-
-	// We operate on int64, since we can easily end up with wraparound issues during one of the
-	// multiplications below when using int32. This would tank the division by WeightSum.
-	int64 WeightSum = 0;
-	for (const FBoneWeight& BW : BoneWeights)
-	{
-		WeightSum += BW.GetRawWeight();
-	}
-
-	if ((InNormalizeType == EBoneWeightNormalizeType::Always && ensure(WeightSum != 0)) ||
-	    WeightSum > FBoneWeight::GetMaxRawWeight())
-	{
-		int64 Correction = 0;
-
-		// Here we treat the raw weight as a 16.16 fixed point value and ensure that the
-		// fraction, which would otherwise be lost through rounding, is carried over to the 
-		// subsequent values to maintain a constant sum to the max weight value.
-		// We do this in descending weight order in an attempt to ensure that weight values 
-		// aren't needlessly lost after scaling.
-		for (FBoneWeight& BW : BoneWeights)
-		{
-			int64 ScaledWeight = int64(BW.GetRawWeight()) * FBoneWeight::GetMaxRawWeight() + Correction;
-			BW.SetRawWeight(uint16(FMath::Min(ScaledWeight / WeightSum, int64(FBoneWeight::GetMaxRawWeight()))));
-			Correction = ScaledWeight - BW.GetRawWeight() * WeightSum;
-		}
-	}
-}
-	
-
 } // namespace AnimationCore
 } // namespace UE
 
