AnimationCore: FEulerTransform scale to init to OneVector

diff --git a/Engine/Source/Runtime/AnimationCore/Public/EulerTransform.h b/Engine/Source/Runtime/AnimationCore/Public/EulerTransform.h
--- a/Engine/Source/Runtime/AnimationCore/Public/EulerTransform.h
+++ b/Engine/Source/Runtime/AnimationCore/Public/EulerTransform.h
@@ -29,7 +29,7 @@ struct ANIMATIONCORE_API FEulerTransform
 	FORCEINLINE FEulerTransform()
 		: Location(ForceInitToZero)
 		, Rotation(ForceInitToZero)
-		, Scale(ForceInitToZero)
+		, Scale(FVector::OneVector)
 	{
 	}
 
