AnimationCore: BoneWeights.  Make conversion to uint16 explicit

diff --git a/Engine/Source/Runtime/AnimationCore/Public/BoneWeights.h b/Engine/Source/Runtime/AnimationCore/Public/BoneWeights.h
--- a/Engine/Source/Runtime/AnimationCore/Public/BoneWeights.h
+++ b/Engine/Source/Runtime/AnimationCore/Public/BoneWeights.h
@@ -64,8 +64,9 @@ public:
 	* bit integer.
 	*/
 	explicit FBoneWeight(FBoneIndexType InBoneIndex, uint8 InWeight)
-	    : BoneIndex(InBoneIndex), RawWeight((InWeight << 8) | InWeight)
+	    : BoneIndex(InBoneIndex)
 	{
+		RawWeight = uint16( (InWeight << 8) | InWeight );
 	}
 
 	/**
@@ -1036,7 +1037,7 @@ TBoneWeights<ContainerAdapter>::Blend(
 		// advance until we hit the end of either array after which we blindly copy the remains.
 		if (BWA.GetBoneIndex() == BWB.GetBoneIndex())
 		{
-			uint16 RawWeight = (BWA.GetRawWeight() * RawBiasA + BWB.GetRawWeight() * RawBiasB) / FBoneWeight::GetMaxRawWeight();
+			uint16 RawWeight = uint16( ((int32)BWA.GetRawWeight() * RawBiasA + (int32)BWB.GetRawWeight() * RawBiasB) / FBoneWeight::GetMaxRawWeight() );
 
 			BoneWeights.Emplace(BWA.GetBoneIndex(), RawWeight);
 			IndexA++;
