AnimationCore: relax the condition for checking the singularity a bit, use 1 - KINDA_SMALL_NUMBER instead of 1 - SMALL_NUMBER

diff --git a/Engine/Source/Runtime/AnimationCore/Private/AnimationCoreLibrary.cpp b/Engine/Source/Runtime/AnimationCore/Private/AnimationCoreLibrary.cpp
--- a/Engine/Source/Runtime/AnimationCore/Private/AnimationCoreLibrary.cpp
+++ b/Engine/Source/Runtime/AnimationCore/Private/AnimationCoreLibrary.cpp
@@ -295,7 +295,7 @@ FVector EulerFromQuat(const FQuat& Rotation, EEulerRotationOrder RotationOrder)
 	{
 		Result.Y = FMath::Asin( - FMath::Clamp<float>( AxisZ.X, -1.f, 1.f ) );
 
-		if ( FMath::Abs( AxisZ.X ) < 1.f - SMALL_NUMBER )
+		if ( FMath::Abs( AxisZ.X ) < 1.f - KINDA_SMALL_NUMBER )
 		{
 			Result.X = FMath::Atan2( AxisZ.Y, AxisZ.Z );
 			Result.Z = FMath::Atan2( AxisY.X, AxisX.X );
@@ -308,10 +308,12 @@ FVector EulerFromQuat(const FQuat& Rotation, EEulerRotationOrder RotationOrder)
 	}
 	else if ( RotationOrder == EEulerRotationOrder::XZY )
 	{
-
+		// using KINDA_SMALL_NUMBER instead of SMALL_NUMBER so that it
+		// better covers singularity cases, a specific case that was causing problems
+		// was Quat(0.122787803968.., -0.122787803968.., -0.696364240320.., 0.696364240320..)
 		Result.Z = FMath::Asin( FMath::Clamp<float>( AxisY.X, -1.f, 1.f ) );
 
-		if ( FMath::Abs( AxisY.X ) < 1.f - SMALL_NUMBER )
+		if ( FMath::Abs( AxisY.X ) < 1.f - KINDA_SMALL_NUMBER )
 		{
 			Result.X = FMath::Atan2( -AxisY.Z, AxisY.Y );
 			Result.Y = FMath::Atan2( -AxisZ.X, AxisX.X );
@@ -326,7 +328,7 @@ FVector EulerFromQuat(const FQuat& Rotation, EEulerRotationOrder RotationOrder)
 	{
 		Result.X = FMath::Asin( FMath::Clamp<float>( AxisZ.Y, -1.f, 1.f ) );
 
-		if ( FMath::Abs( AxisZ.Y ) < 1.f - SMALL_NUMBER )
+		if ( FMath::Abs( AxisZ.Y ) < 1.f - KINDA_SMALL_NUMBER )
 		{
 			Result.Y = FMath::Atan2( -AxisZ.X, AxisZ.Z );
 			Result.Z = FMath::Atan2( -AxisX.Y, AxisY.Y );
@@ -341,7 +343,7 @@ FVector EulerFromQuat(const FQuat& Rotation, EEulerRotationOrder RotationOrder)
 	{
 		Result.Z = FMath::Asin( - FMath::Clamp<float>( AxisX.Y, -1.f, 1.f ) );
 
-		if ( FMath::Abs( AxisX.Y ) < 1.f - SMALL_NUMBER )
+		if ( FMath::Abs( AxisX.Y ) < 1.f - KINDA_SMALL_NUMBER )
 		{
 			Result.X = FMath::Atan2( AxisZ.Y, AxisY.Y );
 			Result.Y = FMath::Atan2( AxisX.Z, AxisX.X );
@@ -356,7 +358,7 @@ FVector EulerFromQuat(const FQuat& Rotation, EEulerRotationOrder RotationOrder)
 	{
 		Result.X = FMath::Asin( - FMath::Clamp<float>( AxisY.Z, -1.f, 1.f ) );
 
-		if ( FMath::Abs( AxisY.Z ) < 1.f - SMALL_NUMBER )
+		if ( FMath::Abs( AxisY.Z ) < 1.f - KINDA_SMALL_NUMBER )
 		{
 			Result.Y = FMath::Atan2( AxisX.Z, AxisZ.Z );
 			Result.Z = FMath::Atan2( AxisY.X, AxisY.Y );
@@ -371,7 +373,7 @@ FVector EulerFromQuat(const FQuat& Rotation, EEulerRotationOrder RotationOrder)
 	{
 		Result.Y = FMath::Asin( FMath::Clamp<float>( AxisX.Z, -1.f, 1.f ) );
 
-		if ( FMath::Abs( AxisX.Z ) < 1.f - SMALL_NUMBER )
+		if ( FMath::Abs( AxisX.Z ) < 1.f - KINDA_SMALL_NUMBER )
 		{
 			Result.X = FMath::Atan2( -AxisY.Z, AxisZ.Z );
 			Result.Z = FMath::Atan2( -AxisX.Y, AxisX.X );
