Truncation fixes for AnimationCore module

diff --git a/Engine/Source/Runtime/AnimationCore/Private/AngularLimit.cpp b/Engine/Source/Runtime/AnimationCore/Private/AngularLimit.cpp
--- a/Engine/Source/Runtime/AnimationCore/Private/AngularLimit.cpp
+++ b/Engine/Source/Runtime/AnimationCore/Private/AngularLimit.cpp
@@ -10,15 +10,13 @@ namespace AnimationCore
 		FQuat DeltaQuat = InRefRotation * InOutQuatRotation.Inverse();
 		FRotator DeltaRotator = DeltaQuat.Rotator();
 		FRotator NewRotator;
-		NewRotator.Pitch = FMath::Clamp<float>(DeltaRotator.Pitch, InLimitMinDegrees.Y, InLimitMaxDegrees.Y);
-		NewRotator.Yaw = FMath::Clamp<float>(DeltaRotator.Yaw, InLimitMinDegrees.X, InLimitMaxDegrees.X);
-		NewRotator.Roll = FMath::Clamp<float>(DeltaRotator.Roll, InLimitMinDegrees.Z, InLimitMaxDegrees.Z);
+		NewRotator.Pitch = FMath::Clamp(DeltaRotator.Pitch, InLimitMinDegrees.Y, InLimitMaxDegrees.Y);
+		NewRotator.Yaw = FMath::Clamp(DeltaRotator.Yaw, InLimitMinDegrees.X, InLimitMaxDegrees.X);
+		NewRotator.Roll = FMath::Clamp(DeltaRotator.Roll, InLimitMinDegrees.Z, InLimitMaxDegrees.Z);
 		DeltaQuat = FQuat(NewRotator);
 
 		InOutQuatRotation = DeltaQuat.Inverse() * InRefRotation;
 
 		return !NewRotator.Equals(DeltaRotator);
 	}
-
 }
-
diff --git a/Engine/Source/Runtime/AnimationCore/Private/AnimationCoreLibrary.cpp b/Engine/Source/Runtime/AnimationCore/Private/AnimationCoreLibrary.cpp
--- a/Engine/Source/Runtime/AnimationCore/Private/AnimationCoreLibrary.cpp
+++ b/Engine/Source/Runtime/AnimationCore/Private/AnimationCoreLibrary.cpp
@@ -52,7 +52,6 @@ FTransform SolveConstraints(const FTransform& CurrentTransform, const FTransform
 	FComponentBlendHelper BlendHelper;
 	FTransform BlendedTransform = CurrentTransform;
 
-	float TotalWeight = 0.f;
 	for (int32 ConstraintIndex = 0; ConstraintIndex < TotalNum; ++ConstraintIndex)
 	{
 		const FTransformConstraint& Constraint = Constraints[ConstraintIndex];
@@ -109,12 +108,12 @@ FQuat SolveAim(const FTransform& CurrentTransform, const FVector& TargetPosition
 
 	if (AimClampInDegree > ZERO_ANIMWEIGHT_THRESH)
 	{
-		float AimClampInRadians = FMath::DegreesToRadians(FMath::Min(AimClampInDegree, 180.f));
-		float DiffAngle = FMath::Acos(FVector::DotProduct(AimVector, ToTarget));
+		double AimClampInRadians = FMath::DegreesToRadians(FMath::Min(AimClampInDegree, 180.0));
+		double DiffAngle = FMath::Acos(FVector::DotProduct(AimVector, ToTarget));
 
 		if (DiffAngle > AimClampInRadians)
 		{
-			check(DiffAngle > 0.f);
+			check(DiffAngle > 0.0);
 
 			FVector DeltaTarget = ToTarget - AimVector;
 			// clamp delta target to within the ratio
@@ -147,7 +146,6 @@ FTransform SolveConstraints(const FTransform& CurrentTransform, const FTransform
 	FMultiTransformBlendHelper BlendHelperInLocalSpace;
 	FTransform BlendedLocalTransform = CurrentTransform.GetRelativeTransform(CurrentParentTransform);
 
-	float TotalWeight = 0.f;
 	for (int32 ConstraintIndex = 0; ConstraintIndex < TotalNum; ++ConstraintIndex)
 	{
 		const FConstraintData& Constraint = Constraints[ConstraintIndex];
@@ -192,17 +190,17 @@ FTransform SolveConstraints(const FTransform& CurrentTransform, const FTransform
 
 FQuat QuatFromEuler(const FVector& XYZAnglesInDegrees, EEulerRotationOrder RotationOrder)
 {
-	float X = FMath::DegreesToRadians(XYZAnglesInDegrees.X);
-	float Y = FMath::DegreesToRadians(XYZAnglesInDegrees.Y);
-	float Z = FMath::DegreesToRadians(XYZAnglesInDegrees.Z);
+	double X = FMath::DegreesToRadians(XYZAnglesInDegrees.X);
+	double Y = FMath::DegreesToRadians(XYZAnglesInDegrees.Y);
+	double Z = FMath::DegreesToRadians(XYZAnglesInDegrees.Z);
 
-	float CosX = FMath::Cos( X * 0.5f );
-	float CosY = FMath::Cos( Y * 0.5f );
-	float CosZ = FMath::Cos( Z * 0.5f );
+	double CosX = FMath::Cos( X * 0.5 );
+	double CosY = FMath::Cos( Y * 0.5 );
+	double CosZ = FMath::Cos( Z * 0.5 );
 
-	float SinX = FMath::Sin( X * 0.5f );
-	float SinY = FMath::Sin( Y * 0.5f );
-	float SinZ = FMath::Sin( Z * 0.5f );
+	double SinX = FMath::Sin( X * 0.5 );
+	double SinY = FMath::Sin( Y * 0.5 );
+	double SinZ = FMath::Sin( Z * 0.5 );
 
 	if ( RotationOrder == EEulerRotationOrder::XYZ )
 	{
@@ -258,51 +256,51 @@ FQuat QuatFromEuler(const FVector& XYZAnglesInDegrees, EEulerRotationOrder Rotat
 
 FVector EulerFromQuat(const FQuat& Rotation, EEulerRotationOrder RotationOrder)
 {
-	float X = Rotation.X;
-	float Y = Rotation.Y;
-	float Z = Rotation.Z;
-	float W = Rotation.W;
-	float X2 = X * 2.f;
-	float Y2 = Y * 2.f;
-	float Z2 = Z * 2.f;
-	float XX2 = X * X2;
-	float XY2 = X * Y2;
-	float XZ2 = X * Z2;
-	float YX2 = Y * X2;
-	float YY2 = Y * Y2;
-	float YZ2 = Y * Z2;
-	float ZX2 = Z * X2;
-	float ZY2 = Z * Y2;
-	float ZZ2 = Z * Z2;
-	float WX2 = W * X2;
-	float WY2 = W * Y2;
-	float WZ2 = W * Z2;
+	double X = Rotation.X;
+	double Y = Rotation.Y;
+	double Z = Rotation.Z;
+	double W = Rotation.W;
+	double X2 = X * 2.0;
+	double Y2 = Y * 2.0;
+	double Z2 = Z * 2.0;
+	double XX2 = X * X2;
+	double XY2 = X * Y2;
+	double XZ2 = X * Z2;
+	double YX2 = Y * X2;
+	double YY2 = Y * Y2;
+	double YZ2 = Y * Z2;
+	double ZX2 = Z * X2;
+	double ZY2 = Z * Y2;
+	double ZZ2 = Z * Z2;
+	double WX2 = W * X2;
+	double WY2 = W * Y2;
+	double WZ2 = W * Z2;
 
 	FVector AxisX, AxisY, AxisZ;
-	AxisX.X = (1.f - (YY2 + ZZ2));
+	AxisX.X = (1.0 - (YY2 + ZZ2));
 	AxisY.X = (XY2 + WZ2);
 	AxisZ.X = (XZ2 - WY2);
 	AxisX.Y = (XY2 - WZ2);
-	AxisY.Y = (1.f - (XX2 + ZZ2));
+	AxisY.Y = (1.0 - (XX2 + ZZ2));
 	AxisZ.Y = (YZ2 + WX2);
 	AxisX.Z = (XZ2 + WY2);
 	AxisY.Z = (YZ2 - WX2);
-	AxisZ.Z = (1.f - (XX2 + YY2));
+	AxisZ.Z = (1.0 - (XX2 + YY2));
 
 	FVector Result = FVector::ZeroVector;
 
 	if ( RotationOrder == EEulerRotationOrder::XYZ )
 	{
-		Result.Y = FMath::Asin( - FMath::Clamp<float>( AxisZ.X, -1.f, 1.f ) );
+		Result.Y = FMath::Asin( - FMath::Clamp( AxisZ.X, -1.0, 1.0 ) );
 
-		if ( FMath::Abs( AxisZ.X ) < 1.f - KINDA_SMALL_NUMBER )
+		if ( FMath::Abs( AxisZ.X ) < 1.0 - DOUBLE_KINDA_SMALL_NUMBER)
 		{
 			Result.X = FMath::Atan2( AxisZ.Y, AxisZ.Z );
 			Result.Z = FMath::Atan2( AxisY.X, AxisX.X );
 		}
 		else
 		{
-			Result.X = 0.f;
+			Result.X = 0.0;
 			Result.Z = FMath::Atan2( -AxisX.Y, AxisY.Y );
 		}
 	}
@@ -311,39 +309,39 @@ FVector EulerFromQuat(const FQuat& Rotation, EEulerRotationOrder RotationOrder)
 		// using KINDA_SMALL_NUMBER instead of SMALL_NUMBER so that it
 		// better covers singularity cases, a specific case that was causing problems
 		// was Quat(0.122787803968.., -0.122787803968.., -0.696364240320.., 0.696364240320..)
-		Result.Z = FMath::Asin( FMath::Clamp<float>( AxisY.X, -1.f, 1.f ) );
+		Result.Z = FMath::Asin( FMath::Clamp( AxisY.X, -1.0, 1.0 ) );
 
-		if ( FMath::Abs( AxisY.X ) < 1.f - KINDA_SMALL_NUMBER )
+		if ( FMath::Abs( AxisY.X ) < 1.0 - DOUBLE_KINDA_SMALL_NUMBER)
 		{
 			Result.X = FMath::Atan2( -AxisY.Z, AxisY.Y );
 			Result.Y = FMath::Atan2( -AxisZ.X, AxisX.X );
 		}
 		else
 		{
-			Result.X = 0.f;
+			Result.X = 0.0;
 			Result.Y = FMath::Atan2( AxisX.Z, AxisZ.Z );
 		}
 	}
 	else if ( RotationOrder == EEulerRotationOrder::YXZ )
 	{
-		Result.X = FMath::Asin( FMath::Clamp<float>( AxisZ.Y, -1.f, 1.f ) );
+		Result.X = FMath::Asin( FMath::Clamp( AxisZ.Y, -1.0, 1.0 ) );
 
-		if ( FMath::Abs( AxisZ.Y ) < 1.f - KINDA_SMALL_NUMBER )
+		if ( FMath::Abs( AxisZ.Y ) < 1.0 - DOUBLE_KINDA_SMALL_NUMBER)
 		{
 			Result.Y = FMath::Atan2( -AxisZ.X, AxisZ.Z );
 			Result.Z = FMath::Atan2( -AxisX.Y, AxisY.Y );
 		}
 		else
 		{
-			Result.Y = 0.f;
+			Result.Y = 0.0;
 			Result.Z = FMath::Atan2( AxisY.X, AxisX.X );
 		}
 	}
 	else if ( RotationOrder == EEulerRotationOrder::YZX )
 	{
-		Result.Z = FMath::Asin( - FMath::Clamp<float>( AxisX.Y, -1.f, 1.f ) );
+		Result.Z = FMath::Asin( - FMath::Clamp( AxisX.Y, -1.0, 1.0 ) );
 
-		if ( FMath::Abs( AxisX.Y ) < 1.f - KINDA_SMALL_NUMBER )
+		if ( FMath::Abs( AxisX.Y ) < 1.0 - DOUBLE_KINDA_SMALL_NUMBER)
 		{
 			Result.X = FMath::Atan2( AxisZ.Y, AxisY.Y );
 			Result.Y = FMath::Atan2( AxisX.Z, AxisX.X );
@@ -351,14 +349,14 @@ FVector EulerFromQuat(const FQuat& Rotation, EEulerRotationOrder RotationOrder)
 		else
 		{
 			Result.X = FMath::Atan2( -AxisY.Z, AxisZ.Z );
-			Result.Y = 0.f;
+			Result.Y = 0.0;
 		}
 	}
 	else if ( RotationOrder == EEulerRotationOrder::ZXY )
 	{
-		Result.X = FMath::Asin( - FMath::Clamp<float>( AxisY.Z, -1.f, 1.f ) );
+		Result.X = FMath::Asin( - FMath::Clamp( AxisY.Z, -1.0, 1.0 ) );
 
-		if ( FMath::Abs( AxisY.Z ) < 1.f - KINDA_SMALL_NUMBER )
+		if ( FMath::Abs( AxisY.Z ) < 1.0 - DOUBLE_KINDA_SMALL_NUMBER)
 		{
 			Result.Y = FMath::Atan2( AxisX.Z, AxisZ.Z );
 			Result.Z = FMath::Atan2( AxisY.X, AxisY.Y );
@@ -366,14 +364,14 @@ FVector EulerFromQuat(const FQuat& Rotation, EEulerRotationOrder RotationOrder)
 		else
 		{
 			Result.Y = FMath::Atan2( -AxisZ.X, AxisX.X );
-			Result.Z = 0.f;
+			Result.Z = 0.0;
 		}
 	}
 	else if ( RotationOrder == EEulerRotationOrder::ZYX )
 	{
-		Result.Y = FMath::Asin( FMath::Clamp<float>( AxisX.Z, -1.f, 1.f ) );
+		Result.Y = FMath::Asin( FMath::Clamp( AxisX.Z, -1.0, 1.0 ) );
 
-		if ( FMath::Abs( AxisX.Z ) < 1.f - KINDA_SMALL_NUMBER )
+		if ( FMath::Abs( AxisX.Z ) < 1.0 - DOUBLE_KINDA_SMALL_NUMBER)
 		{
 			Result.X = FMath::Atan2( -AxisY.Z, AxisZ.Z );
 			Result.Z = FMath::Atan2( -AxisX.Y, AxisX.X );
@@ -381,11 +379,11 @@ FVector EulerFromQuat(const FQuat& Rotation, EEulerRotationOrder RotationOrder)
 		else
 		{
 			Result.X = FMath::Atan2( AxisZ.Y, AxisY.Y );
-			Result.Z = 0.f;
+			Result.Z = 0.0;
 		}
 	}
 
-	return Result * 180.f / PI;
+	return Result * 180.0 / DOUBLE_PI;
 }
 
 FVector ChangeEulerRotationOrder(const FVector& XYZAnglesInDegrees, EEulerRotationOrder SourceRotationOrder, EEulerRotationOrder TargetRotationOrder)
diff --git a/Engine/Source/Runtime/AnimationCore/Private/CCDIK.cpp b/Engine/Source/Runtime/AnimationCore/Private/CCDIK.cpp
--- a/Engine/Source/Runtime/AnimationCore/Private/CCDIK.cpp
+++ b/Engine/Source/Runtime/AnimationCore/Private/CCDIK.cpp
@@ -28,9 +28,9 @@ namespace AnimationCore
 				ToEnd.Normalize();
 				ToTarget.Normalize();
 
-				float RotationLimitPerJointInRadian = FMath::DegreesToRadians(InRotationLimitPerJoints[LinkIndex]);
-				float Angle = FMath::ClampAngle(FMath::Acos(FVector::DotProduct(ToEnd, ToTarget)), -RotationLimitPerJointInRadian, RotationLimitPerJointInRadian);
-				bool bCanRotate = (FMath::Abs(Angle) > KINDA_SMALL_NUMBER) && (!bInEnableRotationLimit || RotationLimitPerJointInRadian > CurrentLink.CurrentAngleDelta);
+				double RotationLimitPerJointInRadian = FMath::DegreesToRadians(InRotationLimitPerJoints[LinkIndex]);
+				double Angle = FMath::ClampAngle(FMath::Acos(FVector::DotProduct(ToEnd, ToTarget)), -RotationLimitPerJointInRadian, RotationLimitPerJointInRadian);
+				bool bCanRotate = (FMath::Abs(Angle) > DOUBLE_KINDA_SMALL_NUMBER) && (!bInEnableRotationLimit || RotationLimitPerJointInRadian > CurrentLink.CurrentAngleDelta);
 				if (bCanRotate)
 				{
 					// check rotation limit first, if fails, just abort
@@ -39,7 +39,7 @@ namespace AnimationCore
 						if (RotationLimitPerJointInRadian < CurrentLink.CurrentAngleDelta + Angle)
 						{
 							Angle = RotationLimitPerJointInRadian - CurrentLink.CurrentAngleDelta;
-							if (Angle <= KINDA_SMALL_NUMBER)
+							if (Angle <= DOUBLE_KINDA_SMALL_NUMBER)
 							{
 								return false;
 							}
@@ -101,7 +101,7 @@ namespace AnimationCore
 			// check how far
 			const FVector TargetPos = TargetPosition;
 			FVector TipPos = InOutChain[TipBoneLinkIndex].Transform.GetLocation();
-			float Distance = FVector::Dist(TargetPos, TipPos);
+			double Distance = FVector::Dist(TargetPos, TipPos);
 			int32 IterationCount = 0;
 			while ((Distance > Precision) && (IterationCount++ < MaxIteration))
 			{
diff --git a/Engine/Source/Runtime/AnimationCore/Private/FABRIK.cpp b/Engine/Source/Runtime/AnimationCore/Private/FABRIK.cpp
--- a/Engine/Source/Runtime/AnimationCore/Private/FABRIK.cpp
+++ b/Engine/Source/Runtime/AnimationCore/Private/FABRIK.cpp
@@ -13,7 +13,7 @@ namespace AnimationCore
 	bool SolveFabrik(TArray<FFABRIKChainLink>& InOutChain, const FVector& TargetPosition, float MaximumReach, float Precision, int32 MaxIterations)
 	{
 		bool bBoneLocationUpdated = false;
-		float const RootToTargetDistSq = FVector::DistSquared(InOutChain[0].Position, TargetPosition);
+		double const RootToTargetDistSq = FVector::DistSquared(InOutChain[0].Position, TargetPosition);
 		int32 const NumChainLinks = InOutChain.Num();
 
 		// FABRIK algorithm - bone translation calculation
@@ -33,7 +33,7 @@ namespace AnimationCore
 			int32 const TipBoneLinkIndex = NumChainLinks - 1;
 
 			// Check distance between tip location and effector location
-			float Slop = FVector::Dist(InOutChain[TipBoneLinkIndex].Position, TargetPosition);
+			double Slop = FVector::Dist(InOutChain[TipBoneLinkIndex].Position, TargetPosition);
 			if (Slop > Precision)
 			{
 				// Set tip bone at end effector location.
diff --git a/Engine/Source/Runtime/AnimationCore/Private/TwoBoneIK.cpp b/Engine/Source/Runtime/AnimationCore/Private/TwoBoneIK.cpp
--- a/Engine/Source/Runtime/AnimationCore/Private/TwoBoneIK.cpp
+++ b/Engine/Source/Runtime/AnimationCore/Private/TwoBoneIK.cpp
@@ -5,14 +5,14 @@
 namespace AnimationCore
 {
 
-void SolveTwoBoneIK(FTransform& InOutRootTransform, FTransform& InOutJointTransform, FTransform& InOutEndTransform, const FVector& JointTarget, const FVector& Effector, bool bAllowStretching, float StartStretchRatio, float MaxStretchScale)
+void SolveTwoBoneIK(FTransform& InOutRootTransform, FTransform& InOutJointTransform, FTransform& InOutEndTransform, const FVector& JointTarget, const FVector& Effector, bool bAllowStretching, double StartStretchRatio, double MaxStretchScale)
 {
-	float LowerLimbLength = (InOutEndTransform.GetLocation() - InOutJointTransform.GetLocation()).Size();
-	float UpperLimbLength = (InOutJointTransform.GetLocation() - InOutRootTransform.GetLocation()).Size();
+	double LowerLimbLength = (InOutEndTransform.GetLocation() - InOutJointTransform.GetLocation()).Size();
+	double UpperLimbLength = (InOutJointTransform.GetLocation() - InOutRootTransform.GetLocation()).Size();
 	SolveTwoBoneIK(InOutRootTransform, InOutJointTransform, InOutEndTransform, JointTarget, Effector, UpperLimbLength, LowerLimbLength, bAllowStretching, StartStretchRatio, MaxStretchScale);
 }
 
-void SolveTwoBoneIK(FTransform& InOutRootTransform, FTransform& InOutJointTransform, FTransform& InOutEndTransform, const FVector& JointTarget, const FVector& Effector, float UpperLimbLength, float LowerLimbLength, bool bAllowStretching, float StartStretchRatio, float MaxStretchScale)
+void SolveTwoBoneIK(FTransform& InOutRootTransform, FTransform& InOutJointTransform, FTransform& InOutEndTransform, const FVector& JointTarget, const FVector& Effector, double UpperLimbLength, double LowerLimbLength, bool bAllowStretching, double StartStretchRatio, double MaxStretchScale)
 {
 	FVector OutJointPos, OutEndPos;
 
@@ -59,30 +59,30 @@ void SolveTwoBoneIK(FTransform& InOutRootTransform, FTransform& InOutJointTransf
 	InOutEndTransform.SetTranslation(OutEndPos);
 }
 
-void SolveTwoBoneIK(const FVector& RootPos, const FVector& JointPos, const FVector& EndPos, const FVector& JointTarget, const FVector& Effector, FVector& OutJointPos, FVector& OutEndPos, bool bAllowStretching, float StartStretchRatio, float MaxStretchScale)
+void SolveTwoBoneIK(const FVector& RootPos, const FVector& JointPos, const FVector& EndPos, const FVector& JointTarget, const FVector& Effector, FVector& OutJointPos, FVector& OutEndPos, bool bAllowStretching, double StartStretchRatio, double MaxStretchScale)
 {
-	float LowerLimbLength = (EndPos - JointPos).Size();
-	float UpperLimbLength = (JointPos - RootPos).Size();
+	const double LowerLimbLength = (EndPos - JointPos).Size();
+	const double UpperLimbLength = (JointPos - RootPos).Size();
 
 	SolveTwoBoneIK(RootPos, JointPos, EndPos, JointTarget, Effector, OutJointPos, OutEndPos, UpperLimbLength, LowerLimbLength, bAllowStretching, StartStretchRatio, MaxStretchScale);
 }
 
-void SolveTwoBoneIK(const FVector& RootPos, const FVector& JointPos, const FVector& EndPos, const FVector& JointTarget, const FVector& Effector, FVector& OutJointPos, FVector& OutEndPos, float UpperLimbLength, float LowerLimbLength, bool bAllowStretching, float StartStretchRatio, float MaxStretchScale)
+void SolveTwoBoneIK(const FVector& RootPos, const FVector& JointPos, const FVector& EndPos, const FVector& JointTarget, const FVector& Effector, FVector& OutJointPos, FVector& OutEndPos, double UpperLimbLength, double LowerLimbLength, bool bAllowStretching, double StartStretchRatio, double MaxStretchScale)
 {
 	// This is our reach goal.
 	FVector DesiredPos = Effector;
 	FVector DesiredDelta = DesiredPos - RootPos;
-	float DesiredLength = DesiredDelta.Size();
+	double DesiredLength = DesiredDelta.Size();
 
 	// Find lengths of upper and lower limb in the ref skeleton.
 	// Use actual sizes instead of ref skeleton, so we take into account translation and scaling from other bone controllers.
-	float MaxLimbLength = LowerLimbLength + UpperLimbLength;
+	double MaxLimbLength = LowerLimbLength + UpperLimbLength;
 
 	// Check to handle case where DesiredPos is the same as RootPos.
 	FVector	DesiredDir;
-	if (DesiredLength < (float)KINDA_SMALL_NUMBER)
+	if (DesiredLength < DOUBLE_KINDA_SMALL_NUMBER)
 	{
-		DesiredLength = (float)KINDA_SMALL_NUMBER;
+		DesiredLength = DOUBLE_KINDA_SMALL_NUMBER;
 		DesiredDir = FVector(1, 0, 0);
 	}
 	else
@@ -92,11 +92,11 @@ void SolveTwoBoneIK(const FVector& RootPos, const FVector& JointPos, const FVect
 
 	// Get joint target (used for defining plane that joint should be in).
 	FVector JointTargetDelta = JointTarget - RootPos;
-	const float JointTargetLengthSqr = JointTargetDelta.SizeSquared();
+	const double JointTargetLengthSqr = JointTargetDelta.SizeSquared();
 
 	// Same check as above, to cover case when JointTarget position is the same as RootPos.
 	FVector JointPlaneNormal, JointBendDir;
-	if (JointTargetLengthSqr < FMath::Square((float)KINDA_SMALL_NUMBER))
+	if (JointTargetLengthSqr < FMath::Square(DOUBLE_KINDA_SMALL_NUMBER))
 	{
 		JointBendDir = FVector(0, 1, 0);
 		JointPlaneNormal = FVector(0, 0, 1);
@@ -107,7 +107,7 @@ void SolveTwoBoneIK(const FVector& RootPos, const FVector& JointPos, const FVect
 
 		// If we are trying to point the limb in the same direction that we are supposed to displace the joint in, 
 		// we have to just pick 2 random vector perp to DesiredDir and each other.
-		if (JointPlaneNormal.SizeSquared() < FMath::Square((float)KINDA_SMALL_NUMBER))
+		if (JointPlaneNormal.SizeSquared() < FMath::Square(DOUBLE_KINDA_SMALL_NUMBER))
 		{
 			DesiredDir.FindBestAxisVectors(JointPlaneNormal, JointBendDir);
 		}
@@ -126,16 +126,16 @@ void SolveTwoBoneIK(const FVector& RootPos, const FVector& JointPos, const FVect
 
 	if (bAllowStretching)
 	{
-		const float ScaleRange = MaxStretchScale - StartStretchRatio;
-		if (ScaleRange > KINDA_SMALL_NUMBER && MaxLimbLength > KINDA_SMALL_NUMBER)
+		const double ScaleRange = MaxStretchScale - StartStretchRatio;
+		if (ScaleRange > DOUBLE_KINDA_SMALL_NUMBER && MaxLimbLength > DOUBLE_KINDA_SMALL_NUMBER)
 		{
-			const float ReachRatio = DesiredLength / MaxLimbLength;
-			const float ScalingFactor = (MaxStretchScale - 1.f) * FMath::Clamp<float>((ReachRatio - StartStretchRatio) / ScaleRange, 0.f, 1.f);
-			if (ScalingFactor > KINDA_SMALL_NUMBER)
+			const double ReachRatio = DesiredLength / MaxLimbLength;
+			const double ScalingFactor = (MaxStretchScale - 1.0) * FMath::Clamp((ReachRatio - StartStretchRatio) / ScaleRange, 0.0, 1.0);
+			if (ScalingFactor > DOUBLE_KINDA_SMALL_NUMBER)
 			{
-				LowerLimbLength *= (1.f + ScalingFactor);
-				UpperLimbLength *= (1.f + ScalingFactor);
-				MaxLimbLength *= (1.f + ScalingFactor);
+				LowerLimbLength *= (1.0 + ScalingFactor);
+				UpperLimbLength *= (1.0 + ScalingFactor);
+				MaxLimbLength *= (1.0 + ScalingFactor);
 			}
 		}
 	}
@@ -153,27 +153,27 @@ void SolveTwoBoneIK(const FVector& RootPos, const FVector& JointPos, const FVect
 	{
 		// So we have a triangle we know the side lengths of. We can work out the angle between DesiredDir and the direction of the upper limb
 		// using the sin rule:
-		const float TwoAB = 2.f * UpperLimbLength * DesiredLength;
+		const double TwoAB = 2.0 * UpperLimbLength * DesiredLength;
 
-		const float CosAngle = (TwoAB != 0.f) ? ((UpperLimbLength*UpperLimbLength) + (DesiredLength*DesiredLength) - (LowerLimbLength*LowerLimbLength)) / TwoAB : 0.f;
+		const double CosAngle = (TwoAB != 0.0) ? ((UpperLimbLength*UpperLimbLength) + (DesiredLength*DesiredLength) - (LowerLimbLength*LowerLimbLength)) / TwoAB : 0.0;
 
 		// If CosAngle is less than 0, the upper arm actually points the opposite way to DesiredDir, so we handle that.
-		const bool bReverseUpperBone = (CosAngle < 0.f);
+		const bool bReverseUpperBone = (CosAngle < 0.0);
 
 		// Angle between upper limb and DesiredDir
 		// ACos clamps internally so we dont need to worry about out-of-range values here.
-		const float Angle = FMath::Acos(CosAngle);
+		const double Angle = FMath::Acos(CosAngle);
 
 		// Now we calculate the distance of the joint from the root -> effector line.
 		// This forms a right-angle triangle, with the upper limb as the hypotenuse.
-		const float JointLineDist = UpperLimbLength * FMath::Sin(Angle);
+		const double JointLineDist = UpperLimbLength * FMath::Sin(Angle);
 
 		// And the final side of that triangle - distance along DesiredDir of perpendicular.
 		// ProjJointDistSqr can't be neg, because JointLineDist must be <= UpperLimbLength because appSin(Angle) is <= 1.
-		const float ProjJointDistSqr = (UpperLimbLength*UpperLimbLength) - (JointLineDist*JointLineDist);
+		const double ProjJointDistSqr = (UpperLimbLength*UpperLimbLength) - (JointLineDist*JointLineDist);
 		// although this shouldn't be ever negative, sometimes Xbox release produces -0.f, causing ProjJointDist to be NaN
 		// so now I branch it. 						
-		float ProjJointDist = (ProjJointDistSqr > 0.f) ? FMath::Sqrt(ProjJointDistSqr) : 0.f;
+		double ProjJointDist = (ProjJointDistSqr > 0.0) ? FMath::Sqrt(ProjJointDistSqr) : 0.0;
 		if (bReverseUpperBone)
 		{
 			ProjJointDist *= -1.f;
diff --git a/Engine/Source/Runtime/AnimationCore/Public/CCDIK.h b/Engine/Source/Runtime/AnimationCore/Public/CCDIK.h
--- a/Engine/Source/Runtime/AnimationCore/Public/CCDIK.h
+++ b/Engine/Source/Runtime/AnimationCore/Public/CCDIK.h
@@ -32,11 +32,11 @@ public:
 	 * They have a zero length distance, so they will inherit this bone's transformation. */
 	TArray<int32> ChildZeroLengthTransformIndices;
 
-	float CurrentAngleDelta;
+	double CurrentAngleDelta;
 
 	FCCDIKChainLink()
 		: TransformIndex(INDEX_NONE)
-		, CurrentAngleDelta(0.f)
+		, CurrentAngleDelta(0.0)
 	{
 	}
 
@@ -44,7 +44,7 @@ public:
 		: Transform(InTransform)
 		, LocalTransform(InLocalTransform)
 		, TransformIndex(InTransformIndex)
-		, CurrentAngleDelta(0.f)
+		, CurrentAngleDelta(0.0)
 	{
 	}
 };
diff --git a/Engine/Source/Runtime/AnimationCore/Public/TwoBoneIK.h b/Engine/Source/Runtime/AnimationCore/Public/TwoBoneIK.h
--- a/Engine/Source/Runtime/AnimationCore/Public/TwoBoneIK.h
+++ b/Engine/Source/Runtime/AnimationCore/Public/TwoBoneIK.h
@@ -26,7 +26,7 @@ namespace AnimationCore
 	 * @param	StartStretchRatio	When should it start stretch -i.e. 1 means its own length without any stretch
 	 * @param	MaxStretchScale		How much it can stretch to in ratio
 	 */
-	ANIMATIONCORE_API void SolveTwoBoneIK(const FVector& RootPos, const FVector& JointPos, const FVector& EndPos, const FVector& JointTarget, const FVector& Effector, FVector& OutJointPos, FVector& OutEndPos, bool bAllowStretching, float StartStretchRatio, float MaxStretchScale);
+	ANIMATIONCORE_API void SolveTwoBoneIK(const FVector& RootPos, const FVector& JointPos, const FVector& EndPos, const FVector& JointTarget, const FVector& Effector, FVector& OutJointPos, FVector& OutEndPos, bool bAllowStretching, double StartStretchRatio, double MaxStretchScale);
 
 	/**
 	* Two Bone IK
@@ -46,7 +46,7 @@ namespace AnimationCore
 	* @param	StartStretchRatio	When should it start stretch -i.e. 1 means its own length without any stretch
 	* @param	MaxStretchScale		How much it can stretch to in ratio
 	*/
-	ANIMATIONCORE_API void SolveTwoBoneIK(const FVector& RootPos, const FVector& JointPos, const FVector& EndPos, const FVector& JointTarget, const FVector& Effector, FVector& OutJointPos, FVector& OutEndPos, float UpperLimbLength, float LowerLimbLength, bool bAllowStretching, float StartStretchRatio, float MaxStretchScale);
+	ANIMATIONCORE_API void SolveTwoBoneIK(const FVector& RootPos, const FVector& JointPos, const FVector& EndPos, const FVector& JointTarget, const FVector& Effector, FVector& OutJointPos, FVector& OutEndPos, double UpperLimbLength, double LowerLimbLength, bool bAllowStretching, double StartStretchRatio, double MaxStretchScale);
 
 	/**
 	* Two Bone IK
@@ -66,7 +66,7 @@ namespace AnimationCore
 	* @param	StartStretchRatio	When should it start stretch -i.e. 1 means its own length without any stretch
 	* @param	MaxStretchScale		How much it can stretch to in ratio
 	*/
-	ANIMATIONCORE_API void SolveTwoBoneIK(FTransform& InOutRootTransform, FTransform& InOutJointTransform, FTransform& InOutEndTransform, const FVector& JointTarget, const FVector& Effector, bool bAllowStretching, float StartStretchRatio, float MaxStretchScale);
+	ANIMATIONCORE_API void SolveTwoBoneIK(FTransform& InOutRootTransform, FTransform& InOutJointTransform, FTransform& InOutEndTransform, const FVector& JointTarget, const FVector& Effector, bool bAllowStretching, double StartStretchRatio, double MaxStretchScale);
 
 	/**
 	* Two Bone IK
@@ -86,5 +86,5 @@ namespace AnimationCore
 	* @param	StartStretchRatio	When should it start stretch -i.e. 1 means its own length without any stretch
 	* @param	MaxStretchScale		How much it can stretch to in ratio
 	*/
-	ANIMATIONCORE_API void SolveTwoBoneIK(FTransform& InOutRootTransform, FTransform& InOutJointTransform, FTransform& InOutEndTransform, const FVector& JointTarget, const FVector& Effector, float UpperLimbLength, float LowerLimbLength, bool bAllowStretching, float StartStretchRatio, float MaxStretchScale);
+	ANIMATIONCORE_API void SolveTwoBoneIK(FTransform& InOutRootTransform, FTransform& InOutJointTransform, FTransform& InOutEndTransform, const FVector& JointTarget, const FVector& Effector, double UpperLimbLength, double LowerLimbLength, bool bAllowStretching, double StartStretchRatio, double MaxStretchScale);
 };
