-Added a new flag (bBlendZeroInfluence) to the FBoneWeightSettings to allow control over how zero-influence bones are treated during the blend. The Blend method will now consider the flag when performing the interpolation.
-Added a Blend function overload that handles barycentric interpolation of 3 skin weights.
-Refactored SetBoneWeightsFromBary (in DynamicVertexSkinWeightsAttribute.h) and ConstructVertexSkinWeightsAttribute (in SelectiveTessellate.cpp) functions to use the newly added Blend function for barycentric interpolation.
-Since the default Blend behavior has changed, make sure that unit tests in BoneWeightTests.cpp are refactored. Added more tests for cases when bBlendZeroInfluence is set to true and when using Blend for barycentric interpolation.

diff --git a/Engine/Plugins/Runtime/GeometryProcessing/Source/DynamicMesh/Private/Operations/SelectiveTessellate.cpp b/Engine/Plugins/Runtime/GeometryProcessing/Source/DynamicMesh/Private/Operations/SelectiveTessellate.cpp
--- a/Engine/Plugins/Runtime/GeometryProcessing/Source/DynamicMesh/Private/Operations/SelectiveTessellate.cpp
+++ b/Engine/Plugins/Runtime/GeometryProcessing/Source/DynamicMesh/Private/Operations/SelectiveTessellate.cpp
@@ -1620,19 +1620,8 @@ namespace SelectiveTessellateLocals
 			Attribute->GetValue(V1, Value1);
 			Attribute->GetValue(V2, Value2);
 			Attribute->GetValue(V3, Value3);
-				
-			// Since FBoneWeights only defines Blend for two inputs, we need to split the barycentric coordinate 
-			// interpolation into two blends. This mimics TDynamicVertexSkinWeightsAttribute::SetBoneWeightsFromBary
-			if (FMath::IsNearlyZero(V + W) == false)
-			{
-				const double BCW = V / (V + W);
-				const FBoneWeights BC = FBoneWeights::Blend(Value2, Value3, (float)BCW);
-				OutValue = FBoneWeights::Blend(Value1, BC, (float)U);
-			}
-			else
-			{
-				OutValue = Value1;
-			}
+			
+			OutValue = FBoneWeights::Blend(Value1, Value2, Value3, (float)U, (float)V, (float)W);
 
 			OutAttribute->SetValue(NewV, OutValue);
 		};
diff --git a/Engine/Source/Runtime/AnimationCore/Private/Tests/BoneWeightsTests.cpp b/Engine/Source/Runtime/AnimationCore/Private/Tests/BoneWeightsTests.cpp
--- a/Engine/Source/Runtime/AnimationCore/Private/Tests/BoneWeightsTests.cpp
+++ b/Engine/Source/Runtime/AnimationCore/Private/Tests/BoneWeightsTests.cpp
@@ -507,6 +507,8 @@ bool FBoneWeightsTestBlend::RunTest(const FString& Parameters)
 
 		FBoneWeightsSettings S;
 		S.SetNormalizeType(EBoneWeightNormalizeType::None);
+		S.SetBlendZeroInfluence(false);
+		
 		auto BW = FBoneWeights::Blend(BW1, BW2, 0.5f, S);
 
 		UTEST_EQUAL(Name, BW.Num(), 3);
@@ -516,6 +518,17 @@ bool FBoneWeightsTestBlend::RunTest(const FString& Parameters)
 		UTEST_EQUAL(Name, BW[1].GetWeight(), 0.5f);
 		UTEST_EQUAL(Name, BW[2].GetBoneIndex(), 2);
 		UTEST_EQUAL(Name, BW[2].GetWeight(), 0.25f);
+
+		S.SetBlendZeroInfluence(true);
+		BW = FBoneWeights::Blend(BW1, BW2, 0.3f, S);
+
+		UTEST_EQUAL(Name, BW.Num(), 3);
+		UTEST_EQUAL(Name, BW[0].GetBoneIndex(), 0);
+		UTEST_EQUAL(Name, BW[0].GetWeight(), 0.7f);
+		UTEST_EQUAL(Name, BW[1].GetBoneIndex(), 1);
+		UTEST_EQUAL(Name, BW[1].GetWeight(), 0.5f);
+		UTEST_EQUAL(Name, BW[2].GetBoneIndex(), 2);
+		UTEST_EQUAL(Name, BW[2].GetWeight(), 0.075f);
 	}
 
 	{
@@ -526,6 +539,8 @@ bool FBoneWeightsTestBlend::RunTest(const FString& Parameters)
 
 		FBoneWeightsSettings S;
 		S.SetNormalizeType(EBoneWeightNormalizeType::None);
+		S.SetBlendZeroInfluence(false);
+
 		auto BW = FBoneWeights::Blend(BW1, BW2, 0.5f, S);
 
 		UTEST_EQUAL(Name, BW.Num(), 3);
@@ -535,6 +550,41 @@ bool FBoneWeightsTestBlend::RunTest(const FString& Parameters)
 		UTEST_EQUAL(Name, BW[1].GetWeight(), 0.5f);
 		UTEST_EQUAL(Name, BW[2].GetBoneIndex(), 1);
 		UTEST_EQUAL(Name, BW[2].GetWeight(), 0.375f);
+
+		S.SetBlendZeroInfluence(true);
+
+		BW = FBoneWeights::Blend(BW1, BW2, 0.5f, S);
+
+		UTEST_EQUAL(Name, BW.Num(), 3);
+		UTEST_EQUAL(Name, BW[0].GetBoneIndex(), 0);
+		UTEST_EQUAL(Name, BW[0].GetWeight(), 0.75f);
+		UTEST_EQUAL(Name, BW[1].GetBoneIndex(), 1);
+		UTEST_EQUAL(Name, BW[1].GetWeight(), 0.375f);
+		UTEST_EQUAL(Name, BW[2].GetBoneIndex(), 2);
+		UTEST_EQUAL(Name, BW[2].GetWeight(), 0.25f);
+	}
+
+	{
+		const TCHAR* Name = TEXT("FBoneWeights::Blend[Barycentric, Two/Two/Three, Two Overlap]");
+
+		auto BW1 = CreateWeights({ {0, 1.0f}, {1, 0.5f} });
+		auto BW2 = CreateWeights({ {0, 0.5f}, {2, 0.5f} });
+		auto BW3 = CreateWeights({ {0, 0.5f}, {1, 0.5f}, {3, 0.1f} });
+
+		FBoneWeightsSettings S;
+		S.SetNormalizeType(EBoneWeightNormalizeType::None);
+		S.SetBlendZeroInfluence(true);
+		auto BW = FBoneWeights::Blend(BW1, BW2, BW3, 0.2f, 0.3f, 0.5f, S);
+
+		UTEST_EQUAL(Name, BW.Num(), 4);
+		UTEST_EQUAL(Name, BW[0].GetBoneIndex(), 0);
+		UTEST_EQUAL(Name, BW[0].GetWeight(), 0.6f);  // 0.2*1.0 + 0.3*0.5 + 0.5*0.5
+		UTEST_EQUAL(Name, BW[1].GetBoneIndex(), 1);
+		UTEST_EQUAL(Name, BW[1].GetWeight(), 0.35f); // 0.2*0.5 + 0.3*0.0 + 0.5*0.5
+		UTEST_EQUAL(Name, BW[2].GetBoneIndex(), 2);
+		UTEST_EQUAL(Name, BW[2].GetWeight(), 0.15f); // 0.2*0.0 + 0.3*0.5 + 0.5*0.0
+		UTEST_EQUAL(Name, BW[3].GetBoneIndex(), 3);
+		UTEST_EQUAL(Name, BW[3].GetWeight(), 0.05f); // 0.2*0.0 + 0.3*0.0 + 0.5*0.1
 	}
 
 	return true;
diff --git a/Engine/Source/Runtime/AnimationCore/Public/BoneWeights.h b/Engine/Source/Runtime/AnimationCore/Public/BoneWeights.h
--- a/Engine/Source/Runtime/AnimationCore/Public/BoneWeights.h
+++ b/Engine/Source/Runtime/AnimationCore/Public/BoneWeights.h
@@ -308,12 +308,29 @@ public:
 		return bHasDefaultBoneIndex;
 	}
 
+	/** Set the way we treat the bones with zero influence during the blend. If true, we assume all 
+	  * bones with no influence have a zero weight (default) and participate in the blend calculation.
+	  * If false, we ignore the bone during the blend.
+	  */
+	void SetBlendZeroInfluence(bool bInBlendZeroInfluence) 
+	{
+		bBlendZeroInfluence = bInBlendZeroInfluence;
+	}
+
+	/** Return true if bones with no influence are considered to have a zero weight during the 
+	  * blend calculation. */
+	bool GetBlendZeroInfluence() const 
+	{
+		return bBlendZeroInfluence;
+	}
+
 private:
 	EBoneWeightNormalizeType NormalizeType = EBoneWeightNormalizeType::Always;
 	int32 MaxWeightCount = MaxInlineBoneWeightCount;
 	uint16 WeightThreshold = 257; // == uint8(1) converted to 16 bit using 'v | v << 8'.
 	FBoneIndexType DefaultBoneIndex = static_cast<FBoneIndexType>(0);
 	bool bHasDefaultBoneIndex = false;
+	bool bBlendZeroInfluence = true; 
 };
 
 /** A null adapter for a bone weight container to use with TBoneWeights. Use as a template to
@@ -653,6 +670,18 @@ public:
 	    const FBoneWeights& InB,
 	    float InBias,
 	    const FBoneWeightsSettings& InSettings = {});
+
+	/** Blend three bone weights via barycentric coordinates using the given settings. 
+	  * Considers the influence from all of the bones. Coordinates should sum to one. 
+	  */
+	static inline FBoneWeights Blend(
+	    const FBoneWeights& InA,
+	    const FBoneWeights& InB,
+		const FBoneWeights& InC,
+	    float InBaryX,
+		float InBaryY,
+		float InBaryZ,
+	    const FBoneWeightsSettings& InSettings = {});
 	
 	// Ranged-based for loop compatibility -- but only the const version.
 	using RangedForConstIteratorType = BoneWeightArrayT::RangedForConstIteratorType;
@@ -1048,6 +1077,21 @@ TBoneWeights<ContainerAdapter>::Blend(
 	const int32 RawBiasB = static_cast<int32>(InBias * static_cast<float>(FBoneWeight::GetMaxRawWeight()));
 	const int32 RawBiasA = FBoneWeight::GetMaxRawWeight() - RawBiasB;
 
+	auto BlendWithZeroInfluenceBone = [&InSettings, &BoneWeights](const FBoneWeight& BW, const int32 RawBias) 
+	{
+		if (InSettings.GetBlendZeroInfluence())  
+		{
+			// Treat the missing bone as having a zero weight 
+			uint16 RawWeight = uint16( (int32)BW.GetRawWeight() * RawBias / FBoneWeight::GetMaxRawWeight() );
+			BoneWeights.Emplace(BW.GetBoneIndex(), RawWeight);
+		}
+		else 
+		{
+			// Ignore the missing bone
+			BoneWeights.Add(BW);
+		}	
+	}; 
+
 	int32 IndexA = 0, IndexB = 0;
 	for (; IndexA < InBoneWeightsA.Num() && IndexB < InBoneWeightsB.Num(); /* */ )
 	{
@@ -1068,23 +1112,23 @@ TBoneWeights<ContainerAdapter>::Blend(
 		}
 		else if (BWA.GetBoneIndex() < BWB.GetBoneIndex())
 		{
-			BoneWeights.Add(BWA);
+			BlendWithZeroInfluenceBone(BWA, RawBiasA);
 			IndexA++;
 		}
 		else
 		{
-			BoneWeights.Add(BWB);
+			BlendWithZeroInfluenceBone(BWB, RawBiasB);
 			IndexB++;
 		}
 	}
 
 	for (; IndexA < InBoneWeightsA.Num(); IndexA++)
 	{
-		BoneWeights.Add(InBoneWeightsA[IndirectIndexA[IndexA]]);
+		BlendWithZeroInfluenceBone(InBoneWeightsA[IndirectIndexA[IndexA]], RawBiasA);
 	}
 	for (; IndexB < InBoneWeightsB.Num(); IndexB++)
 	{
-		BoneWeights.Add(InBoneWeightsB[IndirectIndexB[IndexB]]);
+		BlendWithZeroInfluenceBone(InBoneWeightsB[IndirectIndexB[IndexB]], RawBiasB);		
 	}
 
 	SetBoneWeightsInternal(BoneWeights, InSettings);
@@ -1326,6 +1370,35 @@ FBoneWeights FBoneWeights::Blend(
 	return Result;
 }
 
+FBoneWeights FBoneWeights::Blend(
+	const FBoneWeights& InA,
+	const FBoneWeights& InB,
+	const FBoneWeights& InC,
+	float InBaryX,
+	float InBaryY,
+	float InBaryZ,
+	const FBoneWeightsSettings& InSettings) 
+{
+	if (FMath::IsNearlyZero(InBaryY + InBaryZ) == false)
+	{
+		const float BCW = InBaryZ / (InBaryY + InBaryZ);
+		
+		// Use the default settings to make sure we don't normalize, don't prune bones if we exceed the default limit 
+		// and treat the bones with zero influence as having a zero weight during the first blend
+		FBoneWeightsSettings FirstBlendSettings;
+		FirstBlendSettings.SetNormalizeType(UE::AnimationCore::EBoneWeightNormalizeType::None); 
+		FirstBlendSettings.SetMaxWeightCount(InB.Num() + InC.Num());
+		
+		const FBoneWeights BC = FBoneWeights::Blend(InB, InC, BCW, FirstBlendSettings);
+		const FBoneWeights BCA = FBoneWeights::Blend(BC, InA, InBaryX, InSettings);
+		return BCA;
+	}
+	else
+	{
+		return InA;
+	}
+}
+
 
 } // namespace AnimationCore
 } // namespace UE
diff --git a/Engine/Source/Runtime/GeometryCore/Public/DynamicMesh/DynamicVertexSkinWeightsAttribute.h b/Engine/Source/Runtime/GeometryCore/Public/DynamicMesh/DynamicVertexSkinWeightsAttribute.h
--- a/Engine/Source/Runtime/GeometryCore/Public/DynamicMesh/DynamicVertexSkinWeightsAttribute.h
+++ b/Engine/Source/Runtime/GeometryCore/Public/DynamicMesh/DynamicVertexSkinWeightsAttribute.h
@@ -486,21 +486,45 @@ protected:
 		VertexBoneWeights[SetAttribute] = FBoneWeights::Blend(VertexBoneWeights[AttributeA], VertexBoneWeights[AttributeB], (float)Alpha);
 	}
 
-	/** Set the value at an Attribute to be a barycentric interpolation of three other Attributes */
+	/** Set the value at an Attribute to be a barycentric interpolation of three other Attributes. If one of the 
+	 *  barycentric coordinates is nearly one (within UE_SMALL_NUMBER tolerance) then simply copy over the value of the 
+	 *  corresponding Attribute.
+	 */
 	void SetBoneWeightsFromBary(int SetAttribute, int AttributeA, int AttributeB, int AttributeC, const FVector3d& BaryCoords)
 	{
-		// Since FBoneWeights only defines Blend for two inputs, we need to split the barycentric coordinate interpolation 
-		// into two blends.
-		if (!FMath::IsNearlyZero(BaryCoords.Y + BaryCoords.Z))
-		{
-			const double BCW = BaryCoords.Y / (BaryCoords.Y + BaryCoords.Z);
-			const FBoneWeights BC = FBoneWeights::Blend(VertexBoneWeights[AttributeB], VertexBoneWeights[AttributeC], (float)BCW);
+		const float BaryX = static_cast<float>(BaryCoords.X);
+		const float BaryY  = static_cast<float>(BaryCoords.Y);
+		const float BaryZ = static_cast<float>(BaryCoords.Z);
 
-			VertexBoneWeights[SetAttribute] = FBoneWeights::Blend(VertexBoneWeights[AttributeA], BC, (float)BaryCoords.X);
+		if (FMath::IsNearlyEqual(BaryX, 1.0f))
+		{
+			if (SetAttribute != AttributeA) // avoid unnecessary copy
+			{
+				VertexBoneWeights[SetAttribute] = VertexBoneWeights[AttributeA];
+			}
+		}
+		else if (FMath::IsNearlyEqual(BaryY, 1.0f)) 
+		{
+			if (SetAttribute != AttributeB) 
+			{
+				VertexBoneWeights[SetAttribute] = VertexBoneWeights[AttributeB];
+			}
+		}
+		else if (FMath::IsNearlyEqual(BaryZ, 1.0f)) 
+		{
+			if (SetAttribute != AttributeC) 
+			{
+				VertexBoneWeights[SetAttribute] = VertexBoneWeights[AttributeC];
+			}
 		}
-		else if (SetAttribute != AttributeA)
+		else 
 		{
-			VertexBoneWeights[SetAttribute] = VertexBoneWeights[AttributeA];
+			VertexBoneWeights[SetAttribute] = FBoneWeights::Blend(VertexBoneWeights[AttributeA], 
+																  VertexBoneWeights[AttributeB], 
+																  VertexBoneWeights[AttributeC], 
+																  BaryX,
+																  BaryY,
+																  BaryZ);
 		}
 	}
 };
