Fix System.AnimationCore.BoneWeights.Settings test to match reality.

#rnx

diff --git a/Engine/Source/Runtime/AnimationCore/Private/Tests/BoneWeightsTests.cpp b/Engine/Source/Runtime/AnimationCore/Private/Tests/BoneWeightsTests.cpp
--- a/Engine/Source/Runtime/AnimationCore/Private/Tests/BoneWeightsTests.cpp
+++ b/Engine/Source/Runtime/AnimationCore/Private/Tests/BoneWeightsTests.cpp
@@ -84,7 +84,7 @@ bool FBoneWeightsSettingsTest::RunTest(const FString& Parameters)
 
 	UTEST_EQUAL(TEXT("FBoneWeightsSettings NormalizeType"), S.GetNormalizeType(), EBoneWeightNormalizeType::Always);
 	UTEST_EQUAL(TEXT("FBoneWeightsSettings MaxWeightCount"), S.GetMaxWeightCount(), MaxInlineBoneWeightCount);
-	UTEST_EQUAL(TEXT("FBoneWeightsSettings Weight Threshold"), S.GetRawWeightThreshold(), uint16((1 << 8) | 1));
+	UTEST_EQUAL(TEXT("FBoneWeightsSettings Weight Threshold"), S.GetRawWeightThreshold(), 1);
 
 	return true;
 }
