Fix C++20 compile error. Convert FBoneIndexWithOperators comparison operators to hidden friends so that Lhs is RealBoneIndexType to match Rhs, instead of Lhs being FBoneIndexWithOperators<RealBoneIndexType>, which causes ambiguity. Provide reversed comparison operator overloads for RealBoneIndexType and int32 to match what C++20 will generate. Disable certain comparison operators in FBoneIndexWithOperators for platforms with generated comparison operators.

diff --git a/Engine/Source/Runtime/AnimationCore/Public/BoneIndices.h b/Engine/Source/Runtime/AnimationCore/Public/BoneIndices.h
--- a/Engine/Source/Runtime/AnimationCore/Public/BoneIndices.h
+++ b/Engine/Source/Runtime/AnimationCore/Public/BoneIndices.h
@@ -41,67 +41,103 @@ template<class RealBoneIndexType>
 struct FBoneIndexWithOperators : public FBoneIndexBase
 {
 	// BoneIndexType
-	FORCEINLINE bool operator==(const RealBoneIndexType& Rhs) const
+	FORCEINLINE friend bool operator==(const RealBoneIndexType& Lhs, const RealBoneIndexType& Rhs)
 	{
-		return BoneIndex == GetIntFromComp(Rhs);
+		return GetIntFromComp(Lhs) == GetIntFromComp(Rhs);
 	}
 
-	FORCEINLINE bool operator!=(const RealBoneIndexType& Rhs) const
+#if !PLATFORM_COMPILER_HAS_GENERATED_COMPARISON_OPERATORS
+	FORCEINLINE friend bool operator!=(const RealBoneIndexType& Lhs, const RealBoneIndexType& Rhs)
 	{
-		return BoneIndex != GetIntFromComp(Rhs);
+		return GetIntFromComp(Lhs) != GetIntFromComp(Rhs);
 	}
+#endif // !PLATFORM_COMPILER_HAS_GENERATED_COMPARISON_OPERATORS
 
-	FORCEINLINE bool operator>(const RealBoneIndexType& Rhs) const
+	FORCEINLINE friend bool operator>(const RealBoneIndexType& Lhs, const RealBoneIndexType& Rhs)
 	{
-		return BoneIndex > GetIntFromComp(Rhs);
+		return GetIntFromComp(Lhs) > GetIntFromComp(Rhs);
 	}
 
-	FORCEINLINE bool operator>=(const RealBoneIndexType& Rhs) const
+	FORCEINLINE friend bool operator>=(const RealBoneIndexType& Lhs, const RealBoneIndexType& Rhs)
 	{
-		return BoneIndex >= GetIntFromComp(Rhs);
+		return GetIntFromComp(Lhs) >= GetIntFromComp(Rhs);
 	}
 
-	FORCEINLINE bool operator<(const RealBoneIndexType& Rhs) const
+	FORCEINLINE friend bool operator<(const RealBoneIndexType& Lhs, const RealBoneIndexType& Rhs)
 	{
-		return BoneIndex < GetIntFromComp(Rhs);
+		return GetIntFromComp(Lhs) < GetIntFromComp(Rhs);
 	}
 
-	FORCEINLINE bool operator<=(const RealBoneIndexType& Rhs) const
+	FORCEINLINE friend bool operator<=(const RealBoneIndexType& Lhs, const RealBoneIndexType& Rhs)
 	{
-		return BoneIndex <= GetIntFromComp(Rhs);
+		return GetIntFromComp(Lhs) <= GetIntFromComp(Rhs);
 	}
 
 	// FBoneIndexType
-	FORCEINLINE bool operator==(const int32 Rhs) const
+	FORCEINLINE friend bool operator==(const RealBoneIndexType& Lhs, const int32 Rhs)
 	{
-		return BoneIndex == GetIntFromComp(Rhs);
+		return GetIntFromComp(Lhs) == GetIntFromComp(Rhs);
 	}
 
-	FORCEINLINE bool operator!=(const int32 Rhs) const
+#if !PLATFORM_COMPILER_HAS_GENERATED_COMPARISON_OPERATORS
+	FORCEINLINE friend bool operator!=(const RealBoneIndexType& Lhs, const int32 Rhs)
 	{
-		return BoneIndex != GetIntFromComp(Rhs);
+		return GetIntFromComp(Lhs) != GetIntFromComp(Rhs);
 	}
+#endif // !PLATFORM_COMPILER_HAS_GENERATED_COMPARISON_OPERATORS
 
-	FORCEINLINE bool operator>(const int32 Rhs) const
+	FORCEINLINE friend bool operator>(const RealBoneIndexType& Lhs, const int32 Rhs)
 	{
-		return BoneIndex > GetIntFromComp(Rhs);
+		return GetIntFromComp(Lhs) > GetIntFromComp(Rhs);
 	}
 
-	FORCEINLINE bool operator>=(const int32 Rhs) const
+	FORCEINLINE friend bool operator>=(const RealBoneIndexType& Lhs, const int32 Rhs)
 	{
-		return BoneIndex >= GetIntFromComp(Rhs);
+		return GetIntFromComp(Lhs) >= GetIntFromComp(Rhs);
 	}
 
-	FORCEINLINE bool operator<(const int32 Rhs) const
+	FORCEINLINE friend bool operator<(const RealBoneIndexType& Lhs, const int32 Rhs)
 	{
-		return BoneIndex < GetIntFromComp(Rhs);
+		return GetIntFromComp(Lhs) < GetIntFromComp(Rhs);
 	}
 
-	FORCEINLINE bool operator<=(const int32 Rhs) const
+	FORCEINLINE friend bool operator<=(const RealBoneIndexType& Lhs, const int32 Rhs)
 	{
-		return BoneIndex <= GetIntFromComp(Rhs);
+		return GetIntFromComp(Lhs) <= GetIntFromComp(Rhs);
 	}
 
+#if !PLATFORM_COMPILER_HAS_GENERATED_COMPARISON_OPERATORS
+	FORCEINLINE friend bool operator==(const int32 Lhs, const RealBoneIndexType& Rhs)
+	{
+		return Rhs == Lhs;
+	}
+	
+	FORCEINLINE friend bool operator!=(const int32 Lhs, const RealBoneIndexType& Rhs)
+	{
+		return Rhs != Lhs;
+	}
+
+	FORCEINLINE friend bool operator>(const int32 Lhs, const RealBoneIndexType& Rhs)
+	{
+		return Rhs < Lhs;
+	}
+
+	FORCEINLINE friend bool operator>=(const int32 Lhs, const RealBoneIndexType& Rhs)
+	{
+		return Rhs <= Lhs;
+	}
+
+	FORCEINLINE friend bool operator<(const int32 Lhs, const RealBoneIndexType& Rhs)
+	{
+		return Rhs > Lhs;
+	}
+
+	FORCEINLINE friend bool operator<=(const int32 Lhs, const RealBoneIndexType& Rhs)
+	{
+		return Rhs >= Lhs;
+	}
+#endif // !PLATFORM_COMPILER_HAS_GENERATED_COMPARISON_OPERATORS
+
 	RealBoneIndexType& operator++()
 	{
 		++BoneIndex;
