Bone Weight: Fix incorrect type.

#rnx

diff --git a/Engine/Source/Runtime/AnimationCore/Public/BoneWeights.h b/Engine/Source/Runtime/AnimationCore/Public/BoneWeights.h
--- a/Engine/Source/Runtime/AnimationCore/Public/BoneWeights.h
+++ b/Engine/Source/Runtime/AnimationCore/Public/BoneWeights.h
@@ -68,9 +68,9 @@ public:
 	/**
 	* The maximum raw weight value for a bone influence.
 	*/
-	static constexpr FBoneIndexType GetMaxRawWeight() 
+	static constexpr uint16 GetMaxRawWeight() 
 	{
-		return TNumericLimits<FBoneIndexType>::Max();
+		return TNumericLimits<uint16>::Max();
 	}
 
 	/** A standard predicate we use for sorting by weight, in a descending order of weights */
