Replaced EAutomationTestFlags::ApplicationContextMask with EAutomationTestFlags_ApplicationContextMask.

diff --git a/Engine/Source/Developer/PakFileUtilities/Private/Tests/KeyChainTest.cpp b/Engine/Source/Developer/PakFileUtilities/Private/Tests/KeyChainTest.cpp
--- a/Engine/Source/Developer/PakFileUtilities/Private/Tests/KeyChainTest.cpp
+++ b/Engine/Source/Developer/PakFileUtilities/Private/Tests/KeyChainTest.cpp
@@ -5,7 +5,7 @@
 
 #if WITH_DEV_AUTOMATION_TESTS 
 
-IMPLEMENT_SIMPLE_AUTOMATION_TEST(FKeyChainTest, "System.Core.Misc.KeyChain", EAutomationTestFlags::ApplicationContextMask | EAutomationTestFlags::EngineFilter)
+IMPLEMENT_SIMPLE_AUTOMATION_TEST(FKeyChainTest, "System.Core.Misc.KeyChain", EAutomationTestFlags_ApplicationContextMask | EAutomationTestFlags::EngineFilter)
 
 /**
  * KeyChain tests - implemented in this module because of incorrect dependencies in KeyChainUtilities.h.
diff --git a/Engine/Source/Developer/TraceInsights/Private/Insights/Tests/FunctionalTests/UITests.cpp b/Engine/Source/Developer/TraceInsights/Private/Insights/Tests/FunctionalTests/UITests.cpp
--- a/Engine/Source/Developer/TraceInsights/Private/Insights/Tests/FunctionalTests/UITests.cpp
+++ b/Engine/Source/Developer/TraceInsights/Private/Insights/Tests/FunctionalTests/UITests.cpp
@@ -36,7 +36,7 @@ bool FHideAndShowAllTimingViewTabs::RunTest(const FString& Parameters)
 	return true;
 }
 
-IMPLEMENT_SIMPLE_AUTOMATION_TEST(FMemoryFilterValueConverterTest, "System.Insights.Analysis.MemoryFilterValueConverter", EAutomationTestFlags::ApplicationContextMask | EAutomationTestFlags::EngineFilter)
+IMPLEMENT_SIMPLE_AUTOMATION_TEST(FMemoryFilterValueConverterTest, "System.Insights.Analysis.MemoryFilterValueConverter", EAutomationTestFlags_ApplicationContextMask | EAutomationTestFlags::EngineFilter)
 bool FMemoryFilterValueConverterTest::RunTest(const FString& Parameters)
 {
 	Insights::FMemoryFilterValueConverter Converter;
@@ -72,7 +72,7 @@ bool FMemoryFilterValueConverterTest::RunTest(const FString& Parameters)
 	return !HasAnyErrors();
 }
 
-IMPLEMENT_SIMPLE_AUTOMATION_TEST(FTimeFilterValueConverterTest, "System.Insights.Analysis.TimeFilterValueConverter", EAutomationTestFlags::ApplicationContextMask | EAutomationTestFlags::EngineFilter)
+IMPLEMENT_SIMPLE_AUTOMATION_TEST(FTimeFilterValueConverterTest, "System.Insights.Analysis.TimeFilterValueConverter", EAutomationTestFlags_ApplicationContextMask | EAutomationTestFlags::EngineFilter)
 bool FTimeFilterValueConverterTest::RunTest(const FString& Parameters)
 {
 	Insights::FTimeFilterValueConverter Converter;
diff --git a/Engine/Source/Developer/TraceInsights/Private/Insights/Tests/UnitTests.cpp b/Engine/Source/Developer/TraceInsights/Private/Insights/Tests/UnitTests.cpp
--- a/Engine/Source/Developer/TraceInsights/Private/Insights/Tests/UnitTests.cpp
+++ b/Engine/Source/Developer/TraceInsights/Private/Insights/Tests/UnitTests.cpp
@@ -2,7 +2,7 @@
 #include "CoreMinimal.h"
 #include "Misc/AutomationTest.h"
 
-IMPLEMENT_SIMPLE_AUTOMATION_TEST(FTraceInsightsUnitTest, "System.Insights.UnitTest", EAutomationTestFlags::ApplicationContextMask | EAutomationTestFlags::SmokeFilter)
+IMPLEMENT_SIMPLE_AUTOMATION_TEST(FTraceInsightsUnitTest, "System.Insights.UnitTest", EAutomationTestFlags_ApplicationContextMask | EAutomationTestFlags::SmokeFilter)
 
 bool FTraceInsightsUnitTest::RunTest(const FString& Parameters)
 {
diff --git a/Engine/Source/Developer/TraceServices/Private/Tests/UnitTests.cpp b/Engine/Source/Developer/TraceServices/Private/Tests/UnitTests.cpp
--- a/Engine/Source/Developer/TraceServices/Private/Tests/UnitTests.cpp
+++ b/Engine/Source/Developer/TraceServices/Private/Tests/UnitTests.cpp
@@ -4,7 +4,7 @@
 #include "Common/PagedArray.h"
 #include "Misc/AutomationTest.h"
 
-IMPLEMENT_SIMPLE_AUTOMATION_TEST(FPagedArrayFilteringTest, "System.Insights.Analysis.PagedArrayFiltering", EAutomationTestFlags::ApplicationContextMask | EAutomationTestFlags::SmokeFilter)
+IMPLEMENT_SIMPLE_AUTOMATION_TEST(FPagedArrayFilteringTest, "System.Insights.Analysis.PagedArrayFiltering", EAutomationTestFlags_ApplicationContextMask | EAutomationTestFlags::SmokeFilter)
 
 bool operator!=(const FInt32Interval& lhs, const FInt32Interval& rhs)
 {
diff --git a/Engine/Source/Editor/UnrealEd/Private/Cooker/Algo/GraphConvert.cpp b/Engine/Source/Editor/UnrealEd/Private/Cooker/Algo/GraphConvert.cpp
--- a/Engine/Source/Editor/UnrealEd/Private/Cooker/Algo/GraphConvert.cpp
+++ b/Engine/Source/Editor/UnrealEd/Private/Cooker/Algo/GraphConvert.cpp
@@ -402,7 +402,7 @@ void ConstructPartialTransposeGraph(TConstArrayView<TConstArrayView<FVertex>> Gr
 }
 
 #if WITH_DEV_AUTOMATION_TESTS
-IMPLEMENT_SIMPLE_AUTOMATION_TEST(FCondensationGraphTest, "System.Core.Algo.CondensationGraph", EAutomationTestFlags::ApplicationContextMask | EAutomationTestFlags::EngineFilter);
+IMPLEMENT_SIMPLE_AUTOMATION_TEST(FCondensationGraphTest, "System.Core.Algo.CondensationGraph", EAutomationTestFlags_ApplicationContextMask | EAutomationTestFlags::EngineFilter);
 
 bool FCondensationGraphTest::RunTest(const FString& Parameters)
 {
diff --git a/Engine/Source/Editor/UnrealEd/Private/Cooker/Algo/GraphReachability.cpp b/Engine/Source/Editor/UnrealEd/Private/Cooker/Algo/GraphReachability.cpp
--- a/Engine/Source/Editor/UnrealEd/Private/Cooker/Algo/GraphReachability.cpp
+++ b/Engine/Source/Editor/UnrealEd/Private/Cooker/Algo/GraphReachability.cpp
@@ -686,7 +686,7 @@ void FReachabilityBuilder::RelocateViews(TArray64<FVertex>& InOutTargetBuffer, T
 
 
 #if WITH_DEV_AUTOMATION_TESTS
-IMPLEMENT_SIMPLE_AUTOMATION_TEST(FReachabilityTest, "System.Core.Algo.ReachabilityGraph", EAutomationTestFlags::ApplicationContextMask | EAutomationTestFlags::EngineFilter);
+IMPLEMENT_SIMPLE_AUTOMATION_TEST(FReachabilityTest, "System.Core.Algo.ReachabilityGraph", EAutomationTestFlags_ApplicationContextMask | EAutomationTestFlags::EngineFilter);
 
 bool FReachabilityTest::RunTest(const FString& Parameters)
 {
diff --git a/Engine/Source/Editor/UnrealEd/Private/Cooker/AsyncIODeleteTest.cpp b/Engine/Source/Editor/UnrealEd/Private/Cooker/AsyncIODeleteTest.cpp
--- a/Engine/Source/Editor/UnrealEd/Private/Cooker/AsyncIODeleteTest.cpp
+++ b/Engine/Source/Editor/UnrealEd/Private/Cooker/AsyncIODeleteTest.cpp
@@ -18,7 +18,7 @@
 #include "Templates/UniquePtr.h"
 #include "Trace/Detail/Channel.h"
 
-IMPLEMENT_SIMPLE_AUTOMATION_TEST(FAsyncIODeleteTest, "System.Core.Misc.AsyncIODelete", EAutomationTestFlags::ApplicationContextMask | EAutomationTestFlags::EngineFilter);
+IMPLEMENT_SIMPLE_AUTOMATION_TEST(FAsyncIODeleteTest, "System.Core.Misc.AsyncIODelete", EAutomationTestFlags_ApplicationContextMask | EAutomationTestFlags::EngineFilter);
 bool FAsyncIODeleteTest::RunTest(const FString& Parameters)
 {
 	bool bAsyncEnabled = FAsyncIODelete::AsyncEnabled();
diff --git a/Engine/Source/Editor/UnrealEd/Private/MergeUtils.cpp b/Engine/Source/Editor/UnrealEd/Private/MergeUtils.cpp
--- a/Engine/Source/Editor/UnrealEd/Private/MergeUtils.cpp
+++ b/Engine/Source/Editor/UnrealEd/Private/MergeUtils.cpp
@@ -1207,7 +1207,7 @@ EAssetCommandResult MergeUtils::Merge(const FAssetManualMergeArgs& MergeArgs)
 namespace UE::MergeUtilsTests
 {
 
-	IMPLEMENT_COMPLEX_AUTOMATION_TEST(FMergeWithSelfTests, "ReviewDiffMerge.MergeWithSelf", EAutomationTestFlags::ApplicationContextMask | EAutomationTestFlags::EngineFilter);
+	IMPLEMENT_COMPLEX_AUTOMATION_TEST(FMergeWithSelfTests, "ReviewDiffMerge.MergeWithSelf", EAutomationTestFlags_ApplicationContextMask | EAutomationTestFlags::EngineFilter);
 	void FMergeWithSelfTests::GetTests(TArray<FString>& OutBeautifiedNames, TArray<FString>& OutTestCommands) const
 	{
 		FAssetRegistryModule& AssetRegistryModule = FModuleManager::LoadModuleChecked<FAssetRegistryModule>("AssetRegistry");
diff --git a/Engine/Source/Runtime/Analytics/AnalyticsET/Private/AnalyticsProviderETEventCache.cpp b/Engine/Source/Runtime/Analytics/AnalyticsET/Private/AnalyticsProviderETEventCache.cpp
--- a/Engine/Source/Runtime/Analytics/AnalyticsET/Private/AnalyticsProviderETEventCache.cpp
+++ b/Engine/Source/Runtime/Analytics/AnalyticsET/Private/AnalyticsProviderETEventCache.cpp
@@ -471,7 +471,7 @@ int32 FAnalyticsProviderETEventCache::GetSetPreallocatedPayloadSize() const
 #include "Misc/AutomationTest.h"
 #include <limits>
 
-IMPLEMENT_SIMPLE_AUTOMATION_TEST(FAnalyticsProviderETEventCacheTest, "System.Analytics.AnalyticsETEventCache", EAutomationTestFlags::ApplicationContextMask | EAutomationTestFlags::SmokeFilter)
+IMPLEMENT_SIMPLE_AUTOMATION_TEST(FAnalyticsProviderETEventCacheTest, "System.Analytics.AnalyticsETEventCache", EAutomationTestFlags_ApplicationContextMask | EAutomationTestFlags::SmokeFilter)
 bool FAnalyticsProviderETEventCacheTest::RunTest(const FString& Parameters)
 {
 	// Zero out the DateOffset so we can test against constant strings.
diff --git a/Engine/Source/Runtime/AnimationCore/Private/Tests/BoneWeightsTests.cpp b/Engine/Source/Runtime/AnimationCore/Private/Tests/BoneWeightsTests.cpp
--- a/Engine/Source/Runtime/AnimationCore/Private/Tests/BoneWeightsTests.cpp
+++ b/Engine/Source/Runtime/AnimationCore/Private/Tests/BoneWeightsTests.cpp
@@ -9,7 +9,7 @@
 
 #if WITH_DEV_AUTOMATION_TESTS
 
-IMPLEMENT_SIMPLE_AUTOMATION_TEST(FBoneWeightTestBasic, "System.AnimationCore.BoneWeight.Basic", EAutomationTestFlags::ApplicationContextMask | EAutomationTestFlags::SmokeFilter)
+IMPLEMENT_SIMPLE_AUTOMATION_TEST(FBoneWeightTestBasic, "System.AnimationCore.BoneWeight.Basic", EAutomationTestFlags_ApplicationContextMask | EAutomationTestFlags::SmokeFilter)
 bool FBoneWeightTestBasic::RunTest(const FString& Parameters)
 {
 	using namespace UE::AnimationCore;
@@ -74,7 +74,7 @@ bool FBoneWeightTestBasic::RunTest(const FString& Parameters)
 	return true;
 }
 
-IMPLEMENT_SIMPLE_AUTOMATION_TEST(FBoneWeightsSettingsTest, "System.AnimationCore.BoneWeights.Settings", EAutomationTestFlags::ApplicationContextMask | EAutomationTestFlags::SmokeFilter)
+IMPLEMENT_SIMPLE_AUTOMATION_TEST(FBoneWeightsSettingsTest, "System.AnimationCore.BoneWeights.Settings", EAutomationTestFlags_ApplicationContextMask | EAutomationTestFlags::SmokeFilter)
 bool FBoneWeightsSettingsTest::RunTest(const FString& Parameters)
 {
 	using namespace UE::AnimationCore;
@@ -89,7 +89,7 @@ bool FBoneWeightsSettingsTest::RunTest(const FString& Parameters)
 	return true;
 }
 
-IMPLEMENT_SIMPLE_AUTOMATION_TEST(FBoneWeightsTestBasic, "System.AnimationCore.BoneWeights.Basic", EAutomationTestFlags::ApplicationContextMask | EAutomationTestFlags::SmokeFilter)
+IMPLEMENT_SIMPLE_AUTOMATION_TEST(FBoneWeightsTestBasic, "System.AnimationCore.BoneWeights.Basic", EAutomationTestFlags_ApplicationContextMask | EAutomationTestFlags::SmokeFilter)
 bool FBoneWeightsTestBasic::RunTest(const FString& Parameters)
 {
 	using namespace UE::AnimationCore;
@@ -351,7 +351,7 @@ bool FBoneWeightsTestBasic::RunTest(const FString& Parameters)
 	return true;
 }
 
-IMPLEMENT_SIMPLE_AUTOMATION_TEST(FBoneWeightsTestCreate, "System.AnimationCore.BoneWeights.Create", EAutomationTestFlags::ApplicationContextMask | EAutomationTestFlags::SmokeFilter)
+IMPLEMENT_SIMPLE_AUTOMATION_TEST(FBoneWeightsTestCreate, "System.AnimationCore.BoneWeights.Create", EAutomationTestFlags_ApplicationContextMask | EAutomationTestFlags::SmokeFilter)
 bool FBoneWeightsTestCreate::RunTest(const FString& Parameters)
 {
 	using namespace UE::AnimationCore;
@@ -436,7 +436,7 @@ static UE::AnimationCore::FBoneWeights CreateWeights(
 }
 
 
-IMPLEMENT_SIMPLE_AUTOMATION_TEST(FBoneWeightsTestBlend, "System.AnimationCore.BoneWeights.Blend", EAutomationTestFlags::ApplicationContextMask | EAutomationTestFlags::SmokeFilter)
+IMPLEMENT_SIMPLE_AUTOMATION_TEST(FBoneWeightsTestBlend, "System.AnimationCore.BoneWeights.Blend", EAutomationTestFlags_ApplicationContextMask | EAutomationTestFlags::SmokeFilter)
 bool FBoneWeightsTestBlend::RunTest(const FString& Parameters)
 {
 	using namespace UE::AnimationCore;
