[AnimationCore] Expose EEulerRotationOrder as blueprint type, which should be the enum to use when it comes to rotation order across animtion systems, including control rig.

diff --git a/Engine/Source/Runtime/AnimationCore/Public/EulerTransform.h b/Engine/Source/Runtime/AnimationCore/Public/EulerTransform.h
--- a/Engine/Source/Runtime/AnimationCore/Public/EulerTransform.h
+++ b/Engine/Source/Runtime/AnimationCore/Public/EulerTransform.h
@@ -16,7 +16,7 @@
 class UScriptStruct;
 template <class T> struct TBaseStructure;
 
-UENUM()
+UENUM(BlueprintType)
 enum class EEulerRotationOrder : uint8
 {
 	XYZ,
