[Engine]
* Changed BoneIndices code to use macro for all compare operators instead of templated base class. This does not trip up c++modules/header units

diff --git a/Engine/Source/Runtime/AnimationCore/Public/BoneIndices.h b/Engine/Source/Runtime/AnimationCore/Public/BoneIndices.h
--- a/Engine/Source/Runtime/AnimationCore/Public/BoneIndices.h
+++ b/Engine/Source/Runtime/AnimationCore/Public/BoneIndices.h
@@ -37,145 +37,49 @@ FORCEINLINE int32 GetIntFromComp(const FBoneIndexBase& InComp)
 	return InComp.GetInt();
 }
 
-template<class RealBoneIndexType>
-struct FBoneIndexWithOperators : public FBoneIndexBase
-{
-	// BoneIndexType
-	FORCEINLINE friend bool operator==(const RealBoneIndexType& Lhs, const RealBoneIndexType& Rhs)
-	{
-		return GetIntFromComp(Lhs) == GetIntFromComp(Rhs);
-	}
-
-#if !PLATFORM_COMPILER_HAS_GENERATED_COMPARISON_OPERATORS
-	FORCEINLINE friend bool operator!=(const RealBoneIndexType& Lhs, const RealBoneIndexType& Rhs)
-	{
-		return GetIntFromComp(Lhs) != GetIntFromComp(Rhs);
-	}
-#endif // !PLATFORM_COMPILER_HAS_GENERATED_COMPARISON_OPERATORS
-
-	FORCEINLINE friend bool operator>(const RealBoneIndexType& Lhs, const RealBoneIndexType& Rhs)
-	{
-		return GetIntFromComp(Lhs) > GetIntFromComp(Rhs);
-	}
-
-	FORCEINLINE friend bool operator>=(const RealBoneIndexType& Lhs, const RealBoneIndexType& Rhs)
-	{
-		return GetIntFromComp(Lhs) >= GetIntFromComp(Rhs);
-	}
-
-	FORCEINLINE friend bool operator<(const RealBoneIndexType& Lhs, const RealBoneIndexType& Rhs)
-	{
-		return GetIntFromComp(Lhs) < GetIntFromComp(Rhs);
-	}
-
-	FORCEINLINE friend bool operator<=(const RealBoneIndexType& Lhs, const RealBoneIndexType& Rhs)
-	{
-		return GetIntFromComp(Lhs) <= GetIntFromComp(Rhs);
-	}
-
-	// FBoneIndexType
-	FORCEINLINE friend bool operator==(const RealBoneIndexType& Lhs, const int32 Rhs)
-	{
-		return GetIntFromComp(Lhs) == GetIntFromComp(Rhs);
-	}
+#define UE_BONE_INDEX_COMPARE_OPERATORS(LhsType, RhsType) \
+	FORCEINLINE friend bool operator==(LhsType Lhs, RhsType Rhs) { return GetIntFromComp(Lhs) == GetIntFromComp(Rhs); } \
+	FORCEINLINE friend bool operator>(LhsType Lhs, RhsType Rhs) { return GetIntFromComp(Lhs) > GetIntFromComp(Rhs); } \
+	FORCEINLINE friend bool operator>=(LhsType Lhs, RhsType Rhs) { return GetIntFromComp(Lhs) >= GetIntFromComp(Rhs); } \
+	FORCEINLINE friend bool operator<(LhsType Lhs, RhsType Rhs) { return GetIntFromComp(Lhs) < GetIntFromComp(Rhs); } \
+	FORCEINLINE friend bool operator<=(LhsType Lhs, RhsType Rhs) { return GetIntFromComp(Lhs) <= GetIntFromComp(Rhs); } \
 
 #if !PLATFORM_COMPILER_HAS_GENERATED_COMPARISON_OPERATORS
-	FORCEINLINE friend bool operator!=(const RealBoneIndexType& Lhs, const int32 Rhs)
-	{
-		return GetIntFromComp(Lhs) != GetIntFromComp(Rhs);
-	}
-#endif // !PLATFORM_COMPILER_HAS_GENERATED_COMPARISON_OPERATORS
-
-	FORCEINLINE friend bool operator>(const RealBoneIndexType& Lhs, const int32 Rhs)
-	{
-		return GetIntFromComp(Lhs) > GetIntFromComp(Rhs);
-	}
-
-	FORCEINLINE friend bool operator>=(const RealBoneIndexType& Lhs, const int32 Rhs)
-	{
-		return GetIntFromComp(Lhs) >= GetIntFromComp(Rhs);
-	}
-
-	FORCEINLINE friend bool operator<(const RealBoneIndexType& Lhs, const int32 Rhs)
-	{
-		return GetIntFromComp(Lhs) < GetIntFromComp(Rhs);
-	}
-
-	FORCEINLINE friend bool operator<=(const RealBoneIndexType& Lhs, const int32 Rhs)
-	{
-		return GetIntFromComp(Lhs) <= GetIntFromComp(Rhs);
-	}
-
-#if !PLATFORM_COMPILER_HAS_GENERATED_COMPARISON_OPERATORS
-	FORCEINLINE friend bool operator==(const int32 Lhs, const RealBoneIndexType& Rhs)
-	{
-		return Rhs == Lhs;
-	}
-	
-	FORCEINLINE friend bool operator!=(const int32 Lhs, const RealBoneIndexType& Rhs)
-	{
-		return Rhs != Lhs;
-	}
-
-	FORCEINLINE friend bool operator>(const int32 Lhs, const RealBoneIndexType& Rhs)
-	{
-		return Rhs < Lhs;
-	}
-
-	FORCEINLINE friend bool operator>=(const int32 Lhs, const RealBoneIndexType& Rhs)
-	{
-		return Rhs <= Lhs;
-	}
-
-	FORCEINLINE friend bool operator<(const int32 Lhs, const RealBoneIndexType& Rhs)
-	{
-		return Rhs > Lhs;
-	}
-
-	FORCEINLINE friend bool operator<=(const int32 Lhs, const RealBoneIndexType& Rhs)
-	{
-		return Rhs >= Lhs;
-	}
-#endif // !PLATFORM_COMPILER_HAS_GENERATED_COMPARISON_OPERATORS
-
-	RealBoneIndexType& operator++()
-	{
-		++BoneIndex;
-		return *((RealBoneIndexType*)this);
-	}
-
-	RealBoneIndexType& operator--()
-	{
-		--BoneIndex;
-		return *((RealBoneIndexType*)this);
-	}
-
-	const RealBoneIndexType& operator=(const RealBoneIndexType& Rhs)
-	{
-		BoneIndex = Rhs.BoneIndex;
-		return Rhs;
-	}
-};
+	#define UE_BONE_INDEX_COMPARE_OPERATORS_INVERTED UE_BONE_INDEX_COMPARE_OPERATORS(const int32, const Type&)
+#else
+	#define UE_BONE_INDEX_COMPARE_OPERATORS_INVERTED
+#endif
+
+#define UE_BONE_INDEX_OPERATORS(Type) \
+	UE_BONE_INDEX_COMPARE_OPERATORS(const Type&, const Type&) \
+	UE_BONE_INDEX_COMPARE_OPERATORS(const Type&, const int32) \
+	UE_BONE_INDEX_COMPARE_OPERATORS_INVERTED \
+	Type& operator++() { ++BoneIndex; return *((Type*)this); } \
+	Type& operator--() { --BoneIndex; return *((Type*)this); } \
+	const Type& operator=(const Type& Rhs) { BoneIndex = Rhs.BoneIndex; return Rhs; }
 
 // This represents a compact pose bone index. A compact pose is held by a bone container and can have a different ordering than either the skeleton or skeletal mesh.
-struct FCompactPoseBoneIndex : public FBoneIndexWithOperators < FCompactPoseBoneIndex >
+struct FCompactPoseBoneIndex : public FBoneIndexBase
 {
 public:
 	explicit FCompactPoseBoneIndex(int32 InBoneIndex) { BoneIndex = InBoneIndex; }
+	UE_BONE_INDEX_OPERATORS(FCompactPoseBoneIndex)
 };
 
 // This represents a skeletal mesh bone index which may differ from the skeleton bone index it corresponds to.
-struct FMeshPoseBoneIndex : public FBoneIndexWithOperators < FMeshPoseBoneIndex >
+struct FMeshPoseBoneIndex : public FBoneIndexBase
 {
 public:
 	explicit FMeshPoseBoneIndex(int32 InBoneIndex) { BoneIndex = InBoneIndex; }
+	UE_BONE_INDEX_OPERATORS(FMeshPoseBoneIndex)
 };
 
 // This represents a skeleton bone index which may differ from the skeletal mesh bone index it corresponds to.
-struct FSkeletonPoseBoneIndex : public FBoneIndexWithOperators < FSkeletonPoseBoneIndex >
+struct FSkeletonPoseBoneIndex : public FBoneIndexBase
 {
 public:
 	explicit FSkeletonPoseBoneIndex(int32 InBoneIndex) { BoneIndex = InBoneIndex; }
+	UE_BONE_INDEX_OPERATORS(FSkeletonPoseBoneIndex)
 };
 
 template <typename ValueType>
