Fixed pvs-studio compile warning in BoneIndices

diff --git a/Engine/Source/Runtime/AnimationCore/Public/BoneIndices.h b/Engine/Source/Runtime/AnimationCore/Public/BoneIndices.h
--- a/Engine/Source/Runtime/AnimationCore/Public/BoneIndices.h
+++ b/Engine/Source/Runtime/AnimationCore/Public/BoneIndices.h
@@ -54,9 +54,9 @@ FORCEINLINE int32 GetIntFromComp(const FBoneIndexBase& InComp)
 	UE_BONE_INDEX_COMPARE_OPERATORS(const Type&, const Type&) \
 	UE_BONE_INDEX_COMPARE_OPERATORS(const Type&, const int32) \
 	UE_BONE_INDEX_COMPARE_OPERATORS_INVERTED \
-	Type& operator++() { ++BoneIndex; return *((Type*)this); } \
-	Type& operator--() { --BoneIndex; return *((Type*)this); } \
-	const Type& operator=(const Type& Rhs) { BoneIndex = Rhs.BoneIndex; return Rhs; }
+	Type& operator++() { ++BoneIndex; return *this; } \
+	Type& operator--() { --BoneIndex; return *this; } \
+	Type& operator=(const Type& Rhs) { BoneIndex = Rhs.BoneIndex; return *this; }
 
 // This represents a compact pose bone index. A compact pose is held by a bone container and can have a different ordering than either the skeleton or skeletal mesh.
 struct FCompactPoseBoneIndex : public FBoneIndexBase
